[MLDeformer] Added a new Detail Pose Model. This is a complete rewrite and replacement for the Nearest Neighbor Model. The new model is inherited from the Neural Morph Model and therefore shares all UI/UX and features with it, making the user experience more unified. It only supports the global mode of the Neural Morph Model and it adds some new properties and functionality on top.

This new model will allow you to provide key poses in which we will blend in lost details that the ML reconstruction failed to capture. This gives more detailed folds in clothing etc.

Next to this, the Neural Morph Model directly receives dual quaternion support with this change. Dual quaternion skinning can improve the learned deformations as they are less spikey compared to linear skinned deltas.

This changelist also includes some small bug fixes such as the drawing of deltas not working until you start scrubbing.

#virtualized

diff --git a/Engine/Plugins/Animation/MLDeformer/MLDeformerFramework/Shaders/Private/SkeletonWithWeightedQuats.ush b/Engine/Plugins/Animation/MLDeformer/MLDeformerFramework/Shaders/Private/SkeletonWithWeightedQuats.ush
new file mode 100644
--- /dev/null
+++ b/Engine/Plugins/Animation/MLDeformer/MLDeformerFramework/Shaders/Private/SkeletonWithWeightedQuats.ush
@@ -0,0 +1,230 @@
+// Copyright Epic Games, Inc. All Rights Reserved.
+
+uint {DataInterfaceName}_NumVertices;
+uint {DataInterfaceName}_InputStreamStart;
+uint {DataInterfaceName}_NumBoneInfluences;
+uint {DataInterfaceName}_InputWeightStart;
+uint {DataInterfaceName}_InputWeightStride;
+uint {DataInterfaceName}_InputWeightIndexSize;
+Buffer<float4> {DataInterfaceName}_BoneMatrices;
+Buffer<float4> {DataInterfaceName}_BoneQuats;
+Buffer<uint> {DataInterfaceName}_InputWeightStream;
+Buffer<uint> {DataInterfaceName}_InputWeightLookupStream;
+
+uint ReadNumVertices_{DataInterfaceName}()
+{
+	return {DataInterfaceName}_NumVertices;
+}
+
+#if !GPUSKIN_UNLIMITED_BONE_INFLUENCE
+
+int4 GetBlendIndices_{DataInterfaceName}(uint StreamOffset, uint ExtraInfluenceIndex)
+{
+	return GetBlendIndices({DataInterfaceName}_InputWeightStream, StreamOffset, ExtraInfluenceIndex);
+}
+
+float4 GetBlendWeights_{DataInterfaceName}(uint StreamOffset, uint ExtraInfluenceIndex)
+{
+	return GetBlendWeights({DataInterfaceName}_InputWeightStream, StreamOffset, ExtraInfluenceIndex, {DataInterfaceName}_NumBoneInfluences);
+}
+
+#endif // GPUSKIN_UNLIMITED_BONE_INFLUENCE
+
+float3x4 GetBoneMatrix_{DataInterfaceName}(uint BoneIndex)
+{
+	uint BufferIndex = BoneIndex * 3;
+	return float3x4({DataInterfaceName}_BoneMatrices[BufferIndex], {DataInterfaceName}_BoneMatrices[BufferIndex + 1], {DataInterfaceName}_BoneMatrices[BufferIndex + 2]);
+}
+
+uint ReadNumBones_{DataInterfaceName}(uint VertexIndex)
+{
+#if !ENABLE_DEFORMER_BONES
+	return 1;
+#elif GPUSKIN_UNLIMITED_BONE_INFLUENCE
+	uint BlendOffsetCount = {DataInterfaceName}_InputWeightLookupStream[VertexIndex];
+	return BlendOffsetCount & 0xff;
+#else // !GPUSKIN_UNLIMITED_BONE_INFLUENCE
+	return {DataInterfaceName}_NumBoneInfluences;
+#endif
+}
+
+float3x4 ReadBoneMatrix_{DataInterfaceName}(uint VertexIndex, uint BoneIndex)
+{
+#if !ENABLE_DEFORMER_BONES
+
+	return float3x4(1, 0, 0, 0, 0, 1, 0, 0, 0, 0, 1, 0);
+
+#elif GPUSKIN_UNLIMITED_BONE_INFLUENCE
+	
+	uint BlendOffsetCount = {DataInterfaceName}_InputWeightLookupStream[VertexIndex];
+	int StreamOffset = BlendOffsetCount >> 8;
+
+	int BoneIndexOffset = StreamOffset + ({DataInterfaceName}_InputWeightIndexSize * BoneIndex);
+	int BoneGlobalIndex = {DataInterfaceName}_InputWeightStream[BoneIndexOffset];
+	if ({DataInterfaceName}_InputWeightIndexSize > 1)
+	{
+		BoneGlobalIndex = {DataInterfaceName}_InputWeightStream[BoneIndexOffset + 1] << 8 | BoneGlobalIndex;
+	}
+	
+	return GetBoneMatrix_{DataInterfaceName}(BoneGlobalIndex);
+
+#else // !GPUSKIN_UNLIMITED_BONE_INFLUENCE
+	
+	uint StreamOffset = {DataInterfaceName}_InputWeightStart + (VertexIndex * ({DataInterfaceName}_InputWeightStride / 4));
+	int BlendIndex = GetBlendIndices_{DataInterfaceName}(StreamOffset, BoneIndex / 4)[BoneIndex & 0x3];
+	return GetBoneMatrix_{DataInterfaceName}(BlendIndex);
+
+#endif
+}
+
+float ReadBoneWeight_{DataInterfaceName}(uint VertexIndex, uint BoneIndex)
+{
+#if !ENABLE_DEFORMER_BONES
+
+	return 1;
+
+#elif GPUSKIN_UNLIMITED_BONE_INFLUENCE
+	
+	uint BlendOffsetCount = {DataInterfaceName}_InputWeightLookupStream[VertexIndex];
+	int NumBoneInfluencesLocal = BlendOffsetCount & 0xff;
+	int StreamOffset = BlendOffsetCount >> 8;
+	int BoneIndexSize = {DataInterfaceName}_InputWeightIndexSize & 0xff;
+	int BoneWeightSize = {DataInterfaceName}_InputWeightIndexSize >> 8;
+	int WeightsOffset = StreamOffset + (BoneIndexSize * NumBoneInfluencesLocal);
+	int BoneWeightOffset = WeightsOffset + BoneIndex * BoneWeightSize;
+	
+	if (BoneWeightSize > 1)
+	{
+		return float({DataInterfaceName}_InputWeightStream[BoneWeightOffset + 1] << 8 | {DataInterfaceName}_InputWeightStream[BoneWeightOffset]) / 65535.0;
+	}
+	else
+	{
+		return float({DataInterfaceName}_InputWeightStream[BoneWeightOffset]) / 255.0f;
+	}
+
+#else // !GPUSKIN_UNLIMITED_BONE_INFLUENCE
+	
+	uint StreamOffset = {DataInterfaceName}_InputWeightStart + (VertexIndex * ({DataInterfaceName}_InputWeightStride / 4));
+	float BlendWeight = GetBlendWeights_{DataInterfaceName}(StreamOffset, BoneIndex / 4)[BoneIndex & 0x3];
+	return BlendWeight;
+
+#endif
+}
+
+float3x4 ReadWeightedBoneMatrix_{DataInterfaceName}(uint VertexIndex)
+{
+#if !ENABLE_DEFORMER_BONES
+
+	return float3x4(1, 0, 0, 0, 0, 1, 0, 0, 0, 0, 1, 0);
+
+#elif GPUSKIN_UNLIMITED_BONE_INFLUENCE
+
+	const uint BufferIndex 		= {DataInterfaceName}_InputStreamStart + VertexIndex;
+	const uint BlendOffsetCount = {DataInterfaceName}_InputWeightLookupStream[BufferIndex];
+	return ComputeBoneMatrixWithUnlimitedInfluences({DataInterfaceName}_BoneMatrices, {DataInterfaceName}_InputWeightStream, {DataInterfaceName}_InputWeightIndexSize, BlendOffsetCount);
+
+#else
+
+	const uint StreamOffset = {DataInterfaceName}_InputWeightStart + (VertexIndex * ({DataInterfaceName}_InputWeightStride / 4));
+	const uint NumInfluences = {DataInterfaceName}_NumBoneInfluences;
+	FGPUSkinIndexAndWeight IndicesAndWeights = (FGPUSkinIndexAndWeight)0;
+	if (NumInfluences > 0)
+	{
+		IndicesAndWeights.BlendIndices  = GetBlendIndices_{DataInterfaceName}(StreamOffset, 0);
+		IndicesAndWeights.BlendWeights  = GetBlendWeights_{DataInterfaceName}(StreamOffset, 0);
+	}
+	if (NumInfluences > 4)
+	{
+		IndicesAndWeights.BlendIndices2 = GetBlendIndices_{DataInterfaceName}(StreamOffset, 1);
+		IndicesAndWeights.BlendWeights2 = GetBlendWeights_{DataInterfaceName}(StreamOffset, 1);
+	}
+	return ComputeBoneMatrixWithLimitedInfluences({DataInterfaceName}_BoneMatrices, IndicesAndWeights, NumInfluences > 4);
+#endif
+}
+
+float4 GetBoneQuat_{DataInterfaceName}(uint BoneIndex)
+{
+	return {DataInterfaceName}_BoneQuats[BoneIndex];
+}
+
+float4 GetBoneQuatToBlend_{DataInterfaceName}(uint BoneIndex, float4 BoneQuat0, float Weight)
+{
+	float4 BoneQuat = GetBoneQuat_{DataInterfaceName}(BoneIndex);
+	return dot(BoneQuat, BoneQuat0) >= 0 ? Weight * BoneQuat : -Weight * BoneQuat;
+}
+
+float4 ReadWeightedBoneQuat_{DataInterfaceName}(uint VertexIndex)
+{
+#if !ENABLE_DEFORMER_BONES
+	return float4(0, 0, 0, 1);
+#elif GPUSKIN_UNLIMITED_BONE_INFLUENCE
+	float4 BlendQuat = (float4)0;
+	float4 BoneQuat0 = (float4)0;
+
+	uint BufferIndex = {DataInterfaceName}_InputStreamStart + VertexIndex;
+	uint BlendOffsetCount = {DataInterfaceName}_InputWeightLookupStream[BufferIndex];
+	int NumBoneInfluencesLocal = BlendOffsetCount & 0xff;
+	int IndicesOffset = BlendOffsetCount >> 8;
+	int BoneIndexSize = {DataInterfaceName}_InputWeightIndexSize & 0xff;
+	int BoneWeightSize = {DataInterfaceName}_InputWeightIndexSize >> 8;
+	int WeightsOffset = IndicesOffset + (BoneIndexSize * NumBoneInfluencesLocal);
+
+	for (int InfluenceIndex = 0; InfluenceIndex < NumBoneInfluencesLocal; InfluenceIndex++)
+	{
+		int BoneIndexOffset = IndicesOffset + (BoneIndexSize * InfluenceIndex);
+		int BoneIndex = {DataInterfaceName}_InputWeightStream[BoneIndexOffset];
+		if (BoneIndexSize > 1)
+		{
+			BoneIndex = {DataInterfaceName}_InputWeightStream[BoneIndexOffset + 1] << 8 | BoneIndex;
+		}
+		int BoneWeightOffset = WeightsOffset + (BoneWeightSize * InfluenceIndex);
+		float BoneWeight = float({ DataInterfaceName }_InputWeightStream[BoneWeightOffset]) / 255.0f;
+		if (BoneWeightSize > 1)
+		{
+			BoneWeight = float({DataInterfaceName}_InputWeightStream[BoneWeightOffset + 1] << 8 | {DataInterfaceName}_InputWeightStream[BoneWeightOffset]) / 65535.0;
+		}
+
+		if (InfluenceIndex == 0)
+		{
+			BoneQuat0 = GetBoneQuat_{DataInterfaceName}(BoneIndex);
+			BlendQuat += BoneWeight * BoneQuat0;
+		}
+		else
+		{
+			BlendQuat += GetBoneQuatToBlend_{DataInterfaceName}(BoneIndex, BoneQuat0, BoneWeight);
+		}
+	}
+	return BlendQuat;
+#else
+	float4 BlendQuat = (float4)0;
+	float4 BoneQuat0 = (float4)0;
+
+	uint StreamOffset = {DataInterfaceName}_InputWeightStart + (VertexIndex * ({DataInterfaceName}_InputWeightStride / 4));
+
+	if ({DataInterfaceName}_NumBoneInfluences > 0)
+	{
+		int4 BlendIndices = GetBlendIndices_{DataInterfaceName}(StreamOffset, 0);
+		float4 BlendWeights = GetBlendWeights_{DataInterfaceName}(StreamOffset, 0);
+
+		BoneQuat0 = GetBoneQuat_{DataInterfaceName}(BlendIndices.x);
+		BlendQuat += BlendWeights.x * BoneQuat0;
+		BlendQuat += GetBoneQuatToBlend_{DataInterfaceName}(BlendIndices.y, BoneQuat0, BlendWeights.y);
+		BlendQuat += GetBoneQuatToBlend_{DataInterfaceName}(BlendIndices.z, BoneQuat0, BlendWeights.z);
+		BlendQuat += GetBoneQuatToBlend_{DataInterfaceName}(BlendIndices.w, BoneQuat0, BlendWeights.w);
+	}
+
+	// todo[CF]: make this a permutation when permutation support is working again?
+	if ({DataInterfaceName}_NumBoneInfluences > 4)
+	{
+		int4 BlendIndices = GetBlendIndices_{DataInterfaceName}(StreamOffset, 1);
+		float4 BlendWeights = GetBlendWeights_{DataInterfaceName}(StreamOffset, 1);
+
+		BlendQuat += GetBoneQuatToBlend_{DataInterfaceName}(BlendIndices.x, BoneQuat0, BlendWeights.x);
+		BlendQuat += GetBoneQuatToBlend_{DataInterfaceName}(BlendIndices.y, BoneQuat0, BlendWeights.y);
+		BlendQuat += GetBoneQuatToBlend_{DataInterfaceName}(BlendIndices.z, BoneQuat0, BlendWeights.z);
+		BlendQuat += GetBoneQuatToBlend_{DataInterfaceName}(BlendIndices.w, BoneQuat0, BlendWeights.w);
+	}
+
+	return BlendQuat;
+#endif
+}
diff --git a/Engine/Plugins/Animation/MLDeformer/MLDeformerFramework/Source/MLDeformerFramework/Private/DeformerGraphDataInterfaceSkeletonWithQuats.cpp b/Engine/Plugins/Animation/MLDeformer/MLDeformerFramework/Source/MLDeformerFramework/Private/DeformerGraphDataInterfaceSkeletonWithQuats.cpp
new file mode 100644
--- /dev/null
+++ b/Engine/Plugins/Animation/MLDeformer/MLDeformerFramework/Source/MLDeformerFramework/Private/DeformerGraphDataInterfaceSkeletonWithQuats.cpp
@@ -0,0 +1,316 @@
+// Copyright Epic Games, Inc. All Rights Reserved.
+
+#include "DeformerGraphDataInterfaceSkeletonWithQuats.h"
+#include "ComputeFramework/ComputeKernelPermutationVector.h"
+#include "ComputeFramework/ShaderParamTypeDefinition.h"
+#include "OptimusDataDomain.h"
+#include "RenderGraphBuilder.h"
+#include "Rendering/SkeletalMeshLODRenderData.h"
+#include "Rendering/SkeletalMeshRenderData.h"
+#include "ShaderCompilerCore.h"
+#include "ShaderParameterMetadataBuilder.h"
+#include "SkeletalMeshDeformerHelpers.h"
+#include "SkeletalRenderPublic.h"
+
+#include UE_INLINE_GENERATED_CPP_BY_NAME(DeformerGraphDataInterfaceSkeletonWithQuats)
+
+FString UDeformerGraphSkeletonWithWeightedQuatsDataInterface::GetDisplayName() const
+{
+	return TEXT("MLD Skeleton with Weighted Quats");
+}
+
+TArray<FOptimusCDIPinDefinition> UDeformerGraphSkeletonWithWeightedQuatsDataInterface::GetPinDefinitions() const
+{
+	TArray<FOptimusCDIPinDefinition> Defs;
+	Defs.Add({"NumBones", "ReadNumBones", Optimus::DomainName::Vertex, "ReadNumVertices"});
+	Defs.Add({"BoneMatrix", "ReadBoneMatrix", {{Optimus::DomainName::Vertex, "ReadNumVertices"}, {Optimus::DomainName::Bone, "ReadNumBones"}}});
+	Defs.Add({"BoneWeight", "ReadBoneWeight", {{Optimus::DomainName::Vertex, "ReadNumVertices"}, {Optimus::DomainName::Bone, "ReadNumBones"}}});
+	Defs.Add({"WeightedBoneMatrix", "ReadWeightedBoneMatrix", Optimus::DomainName::Vertex, "ReadNumVertices" });
+	Defs.Add({"WeightedBoneQuat", "ReadWeightedBoneQuat", Optimus::DomainName::Vertex, "ReadNumVertices" });
+	return Defs;
+}
+
+TSubclassOf<UActorComponent> UDeformerGraphSkeletonWithWeightedQuatsDataInterface::GetRequiredComponentClass() const
+{
+	return USkinnedMeshComponent::StaticClass();
+}
+
+void UDeformerGraphSkeletonWithWeightedQuatsDataInterface::GetSupportedInputs(TArray<FShaderFunctionDefinition>& OutFunctions) const
+{
+	OutFunctions.AddDefaulted_GetRef()
+		.SetName(TEXT("ReadNumVertices"))
+		.AddReturnType(EShaderFundamentalType::Uint);
+
+	OutFunctions.AddDefaulted_GetRef()
+		.SetName(TEXT("ReadNumBones"))
+		.AddReturnType(EShaderFundamentalType::Uint)
+		.AddParam(EShaderFundamentalType::Uint);
+
+	OutFunctions.AddDefaulted_GetRef()
+		.SetName(TEXT("ReadBoneMatrix"))
+		.AddReturnType(EShaderFundamentalType::Float, 3, 4)
+		.AddParam(EShaderFundamentalType::Uint)
+		.AddParam(EShaderFundamentalType::Uint);
+
+	OutFunctions.AddDefaulted_GetRef()
+		.SetName(TEXT("ReadBoneWeight"))
+		.AddReturnType(EShaderFundamentalType::Float)
+		.AddParam(EShaderFundamentalType::Uint)
+		.AddParam(EShaderFundamentalType::Uint);
+
+	OutFunctions.AddDefaulted_GetRef()
+		.SetName(TEXT("ReadWeightedBoneMatrix"))
+		.AddReturnType(EShaderFundamentalType::Float, 3, 4)
+		.AddParam(EShaderFundamentalType::Uint);
+
+	OutFunctions.AddDefaulted_GetRef()
+		.SetName(TEXT("ReadWeightedBoneQuat"))
+		.AddReturnType(EShaderFundamentalType::Float, 4)
+		.AddParam(EShaderFundamentalType::Uint);	
+}
+
+BEGIN_SHADER_PARAMETER_STRUCT(FSkeletonWithQuatsDataInterfaceParameters, )
+	SHADER_PARAMETER(uint32, NumVertices)
+	SHADER_PARAMETER(uint32, InputStreamStart)
+	SHADER_PARAMETER(uint32, NumBoneInfluences)
+	SHADER_PARAMETER(uint32, InputWeightStart)
+	SHADER_PARAMETER(uint32, InputWeightStride)
+	SHADER_PARAMETER(uint32, InputWeightIndexSize)
+	SHADER_PARAMETER_SRV(Buffer<float4>, BoneMatrices)
+	SHADER_PARAMETER_RDG_BUFFER_SRV(Buffer<float4>, BoneQuats)
+	SHADER_PARAMETER_SRV(Buffer<uint>, InputWeightStream)
+	SHADER_PARAMETER_SRV(Buffer<uint>, InputWeightLookupStream)
+END_SHADER_PARAMETER_STRUCT()
+
+void UDeformerGraphSkeletonWithWeightedQuatsDataInterface::GetShaderParameters(TCHAR const* UID, FShaderParametersMetadataBuilder& InOutBuilder, FShaderParametersMetadataAllocations& InOutAllocations) const
+{
+	InOutBuilder.AddNestedStruct<FSkeletonWithQuatsDataInterfaceParameters>(UID);
+}
+
+void UDeformerGraphSkeletonWithWeightedQuatsDataInterface::GetPermutations(FComputeKernelPermutationVector& OutPermutationVector) const
+{
+	OutPermutationVector.AddPermutation(TEXT("ENABLE_DEFORMER_BONES"), 2);
+	OutPermutationVector.AddPermutation(TEXT("GPUSKIN_UNLIMITED_BONE_INFLUENCE"), 2);
+	OutPermutationVector.AddPermutation(TEXT("GPUSKIN_BONE_INDEX_UINT16"), 2);
+	OutPermutationVector.AddPermutation(TEXT("GPUSKIN_BONE_WEIGHTS_UINT16"), 2);
+}
+
+TCHAR const* UDeformerGraphSkeletonWithWeightedQuatsDataInterface::TemplateFilePath = TEXT("/Plugin/MLDeformerFramework/Private/SkeletonWithWeightedQuats.ush");
+
+TCHAR const* UDeformerGraphSkeletonWithWeightedQuatsDataInterface::GetShaderVirtualPath() const
+{
+	return TemplateFilePath;
+}
+
+void UDeformerGraphSkeletonWithWeightedQuatsDataInterface::GetShaderHash(FString& InOutKey) const
+{
+	GetShaderFileHash(TemplateFilePath, EShaderPlatform::SP_PCD3D_SM5).AppendString(InOutKey);
+}
+
+void UDeformerGraphSkeletonWithWeightedQuatsDataInterface::GetHLSL(FString& OutHLSL, FString const& InDataInterfaceName) const
+{
+	TMap<FString, FStringFormatArg> TemplateArgs =
+	{
+		{ TEXT("DataInterfaceName"), InDataInterfaceName },
+	};
+
+	FString TemplateFile;
+	LoadShaderSourceFile(TemplateFilePath, EShaderPlatform::SP_PCD3D_SM5, &TemplateFile, nullptr);
+	OutHLSL += FString::Format(*TemplateFile, TemplateArgs);
+}
+
+UComputeDataProvider* UDeformerGraphSkeletonWithWeightedQuatsDataInterface::CreateDataProvider(TObjectPtr<UObject> InBinding, uint64 InInputMask, uint64 InOutputMask) const
+{
+	UDeformerGraphSkeletonWithWeightedQuatsDataProvider* Provider = NewObject<UDeformerGraphSkeletonWithWeightedQuatsDataProvider>();
+	Provider->SkinnedMesh = Cast<USkinnedMeshComponent>(InBinding);
+	return Provider;
+}
+
+FComputeDataProviderRenderProxy* UDeformerGraphSkeletonWithWeightedQuatsDataProvider::GetRenderProxy()
+{
+	return new UE::MLDeformer::FDeformerGraphSkeletonWithWeightedQuatsDataProviderProxy(SkinnedMesh);
+}
+
+namespace UE::MLDeformer
+{
+	FDeformerGraphSkeletonWithWeightedQuatsDataProviderProxy::FDeformerGraphSkeletonWithWeightedQuatsDataProviderProxy(USkinnedMeshComponent* SkinnedMeshComponent)
+	{
+		SkeletalMeshObject = SkinnedMeshComponent != nullptr ? SkinnedMeshComponent->MeshObject : nullptr;
+		BoneRevisionNumber = SkinnedMeshComponent != nullptr ? SkinnedMeshComponent->GetBoneTransformRevisionNumber() : 0;
+	}
+
+	bool FDeformerGraphSkeletonWithWeightedQuatsDataProviderProxy::IsValid(FValidationData const& InValidationData) const
+	{
+		if (InValidationData.ParameterStructSize != sizeof(FParameters))
+		{
+			return false;
+		}
+		if (SkeletalMeshObject == nullptr)
+		{
+			return false;
+		}
+		if (SkeletalMeshObject->GetSkeletalMeshRenderData().LODRenderData[SkeletalMeshObject->GetLOD()].RenderSections.Num() != InValidationData.NumInvocations)
+		{
+			return false;
+		}
+
+		return true;
+	}
+
+	struct FSkeletonWithQuatsDataInterfacePermutationIds
+	{
+		uint32 EnableDeformerBones = 0;
+		uint32 UnlimitedBoneInfluence = 0;
+		uint32 BoneIndexUint16 = 0;
+		uint32 BoneWeightsUint16 = 0;
+
+		FSkeletonWithQuatsDataInterfacePermutationIds(FComputeKernelPermutationVector const& PermutationVector)
+		{
+			{
+				static FString Name(TEXT("ENABLE_DEFORMER_BONES"));
+				static uint32 Hash = GetTypeHash(Name);
+				EnableDeformerBones = PermutationVector.GetPermutationBits(Name, Hash, 1);
+			}
+			{
+				static FString Name(TEXT("GPUSKIN_UNLIMITED_BONE_INFLUENCE"));
+				static uint32 Hash = GetTypeHash(Name);
+				UnlimitedBoneInfluence = PermutationVector.GetPermutationBits(Name, Hash, 1);
+			}
+			{
+				static FString Name(TEXT("GPUSKIN_BONE_INDEX_UINT16"));
+				static uint32 Hash = GetTypeHash(Name);
+				BoneIndexUint16 = PermutationVector.GetPermutationBits(Name, Hash, 1);
+			}
+			{
+				static FString Name(TEXT("GPUSKIN_BONE_WEIGHTS_UINT16"));
+				static uint32 Hash = GetTypeHash(Name);
+				BoneWeightsUint16 = PermutationVector.GetPermutationBits(Name, Hash, 1);
+			}
+		}
+	};
+
+	void FDeformerGraphSkeletonWithWeightedQuatsDataProviderProxy::GatherPermutations(FPermutationData& InOutPermutationData) const
+	{
+		const int32 LodIndex = SkeletalMeshObject->GetLOD();
+		FSkeletalMeshRenderData const& SkeletalMeshRenderData = SkeletalMeshObject->GetSkeletalMeshRenderData();
+		FSkeletalMeshLODRenderData const* LodRenderData = &SkeletalMeshRenderData.LODRenderData[LodIndex];
+
+		FSkeletonWithQuatsDataInterfacePermutationIds PermutationIds(InOutPermutationData.PermutationVector);
+		for (int32 InvocationIndex = 0; InvocationIndex < InOutPermutationData.NumInvocations; ++InvocationIndex)
+		{
+			FSkelMeshRenderSection const& RenderSection = LodRenderData->RenderSections[InvocationIndex];
+
+			const bool bPreviousFrame = false;
+			FRHIShaderResourceView* BoneBufferSRV = FSkeletalMeshDeformerHelpers::GetBoneBufferForReading(SkeletalMeshObject, LodIndex, InvocationIndex, bPreviousFrame);
+
+			FSkinWeightVertexBuffer const* WeightBuffer = LodRenderData->GetSkinWeightVertexBuffer();
+			check(WeightBuffer != nullptr);
+			FRHIShaderResourceView* SkinWeightBufferSRV = WeightBuffer->GetDataVertexBuffer()->GetSRV();
+			const bool bUnlimitedBoneInfluences = WeightBuffer->GetBoneInfluenceType() == GPUSkinBoneInfluenceType::UnlimitedBoneInfluence;
+			FRHIShaderResourceView* InputWeightLookupStreamSRV = bUnlimitedBoneInfluences ? WeightBuffer->GetLookupVertexBuffer()->GetSRV() : nullptr;
+			const bool bValidBones = (BoneBufferSRV != nullptr) && (SkinWeightBufferSRV != nullptr) && (!bUnlimitedBoneInfluences || InputWeightLookupStreamSRV != nullptr);
+			const bool bUse16BitBoneIndex = WeightBuffer->Use16BitBoneIndex();
+			const bool bUse16BitBoneWeights = WeightBuffer->Use16BitBoneWeight();
+
+			InOutPermutationData.PermutationIds[InvocationIndex] |= (bValidBones ? PermutationIds.EnableDeformerBones : 0);
+			InOutPermutationData.PermutationIds[InvocationIndex] |= (bUnlimitedBoneInfluences ? PermutationIds.UnlimitedBoneInfluence : 0);
+			InOutPermutationData.PermutationIds[InvocationIndex] |= (bUse16BitBoneIndex ? PermutationIds.BoneIndexUint16 : 0);
+			InOutPermutationData.PermutationIds[InvocationIndex] |= (bUse16BitBoneWeights ? PermutationIds.BoneWeightsUint16 : 0);
+		}
+	}
+
+	void FDeformerGraphSkeletonWithWeightedQuatsDataProviderProxy::AllocateResources(FRDGBuilder& GraphBuilder)
+	{
+		TArray<FQuat4f> RefToLocalQuats;
+		CacheRefToLocalQuats(RefToLocalQuats);
+
+		const int32 LodIndex = SkeletalMeshObject->GetLOD();
+		FSkeletalMeshRenderData const& SkeletalMeshRenderData = SkeletalMeshObject->GetSkeletalMeshRenderData();
+		FSkeletalMeshLODRenderData const* LodRenderData = &SkeletalMeshRenderData.LODRenderData[LodIndex];
+		const int32 NumSections = LodRenderData->RenderSections.Num();
+		PerSectionRefToLocalQuats.SetNum(NumSections);
+		PerSectionRefToLocalQuatsSRVs.SetNum(NumSections);
+		for (int32 SectionIndex = 0; SectionIndex < NumSections; SectionIndex++)
+		{
+			TArray<FQuat4f> &SectionRefToLocalQuats = PerSectionRefToLocalQuats[SectionIndex];
+			FRDGBufferSRVRef &SectionRefToLocalQuatsSRV = PerSectionRefToLocalQuatsSRVs[SectionIndex];
+			const FSkelMeshRenderSection& Section = LodRenderData->RenderSections[SectionIndex];
+			const uint32 LocalNumBones = Section.BoneMap.Num();
+
+			SectionRefToLocalQuats.SetNum(LocalNumBones);
+			for (uint32 BoneIndex = 0; BoneIndex < LocalNumBones; BoneIndex++)
+			{
+				const FBoneIndexType RefToLocalIndex = Section.BoneMap[BoneIndex];
+				check(RefToLocalQuats.IsValidIndex(RefToLocalIndex));
+				SectionRefToLocalQuats[BoneIndex] = RefToLocalQuats[RefToLocalIndex];
+			}
+
+			if (LocalNumBones <= 0)
+			{
+				SectionRefToLocalQuatsSRV = GraphBuilder.CreateSRV(GraphBuilder.RegisterExternalBuffer(GWhiteVertexBufferWithRDG->Buffer), PF_A32B32G32R32F);
+			}
+			else
+			{
+				FRDGBufferRef SectionRefToLocalQuatsBuffer = GraphBuilder.CreateBuffer(FRDGBufferDesc::CreateBufferDesc(sizeof(FQuat4f), LocalNumBones), TEXT("SectionRefToLocalQuats"));
+				SectionRefToLocalQuatsSRV = GraphBuilder.CreateSRV(SectionRefToLocalQuatsBuffer, PF_A32B32G32R32F);
+				GraphBuilder.QueueBufferUpload(SectionRefToLocalQuatsBuffer, SectionRefToLocalQuats.GetData(), LocalNumBones * sizeof(FQuat4f));
+			}
+		}
+	}
+
+	void FDeformerGraphSkeletonWithWeightedQuatsDataProviderProxy::GatherDispatchData(FDispatchData const& InDispatchData)
+	{
+		const int32 LodIndex = SkeletalMeshObject->GetLOD();
+		FSkeletalMeshRenderData const& SkeletalMeshRenderData = SkeletalMeshObject->GetSkeletalMeshRenderData();
+		FSkeletalMeshLODRenderData const* LodRenderData = &SkeletalMeshRenderData.LODRenderData[LodIndex];
+
+		FRHIShaderResourceView* NullSRVBinding = GWhiteVertexBufferWithSRV->ShaderResourceViewRHI.GetReference();
+
+		const TStridedView<FParameters> ParameterArray = MakeStridedParameterView<FParameters>(InDispatchData);
+		for (int32 InvocationIndex = 0; InvocationIndex < ParameterArray.Num(); ++InvocationIndex)
+		{
+			FSkelMeshRenderSection const& RenderSection = LodRenderData->RenderSections[InvocationIndex];
+
+			const bool bPreviousFrame = false;
+			FRHIShaderResourceView* BoneBufferSRV = FSkeletalMeshDeformerHelpers::GetBoneBufferForReading(SkeletalMeshObject, LodIndex, InvocationIndex, bPreviousFrame);
+
+			FSkinWeightVertexBuffer const* WeightBuffer = LodRenderData->GetSkinWeightVertexBuffer();
+			check(WeightBuffer != nullptr);
+			FRHIShaderResourceView* SkinWeightBufferSRV = WeightBuffer->GetDataVertexBuffer()->GetSRV();
+			const bool bUnlimitedBoneInfluences = WeightBuffer->GetBoneInfluenceType() == GPUSkinBoneInfluenceType::UnlimitedBoneInfluence;
+			FRHIShaderResourceView* InputWeightLookupStreamSRV = bUnlimitedBoneInfluences ? WeightBuffer->GetLookupVertexBuffer()->GetSRV() : nullptr;
+
+			FParameters& Parameters = ParameterArray[InvocationIndex];
+			Parameters.NumVertices = RenderSection.NumVertices;
+			Parameters.InputStreamStart = RenderSection.BaseVertexIndex;
+			Parameters.NumBoneInfluences = WeightBuffer->GetMaxBoneInfluences();
+			Parameters.InputWeightStart = (WeightBuffer->GetConstantInfluencesVertexStride() * RenderSection.GetVertexBufferIndex()) / sizeof(float);
+			Parameters.InputWeightStride = WeightBuffer->GetConstantInfluencesVertexStride();
+			Parameters.InputWeightIndexSize = WeightBuffer->GetBoneIndexByteSize() | (WeightBuffer->GetBoneWeightByteSize() << 8);
+			Parameters.BoneMatrices = BoneBufferSRV != nullptr ? BoneBufferSRV : NullSRVBinding;
+			Parameters.BoneQuats = PerSectionRefToLocalQuatsSRVs[InvocationIndex];
+			Parameters.InputWeightStream = SkinWeightBufferSRV != nullptr ? SkinWeightBufferSRV : NullSRVBinding;
+			Parameters.InputWeightLookupStream = InputWeightLookupStreamSRV != nullptr ? InputWeightLookupStreamSRV : NullSRVBinding;
+		}
+	}
+
+	void FDeformerGraphSkeletonWithWeightedQuatsDataProviderProxy::CacheRefToLocalQuats(TArray<FQuat4f>& OutRefToLocalQuats) const
+	{
+		OutRefToLocalQuats.Reset();
+		if (SkeletalMeshObject == nullptr)
+		{
+			return;
+		}
+
+		const TArray<FMatrix44f>& RefToLocalMatrices = SkeletalMeshObject->GetReferenceToLocalMatrices();
+		const int32 NumQuats = RefToLocalMatrices.Num();
+		OutRefToLocalQuats.SetNumUninitialized(NumQuats);
+
+		// It is theoretically more efficient to compute q_ref ^ {-1} * q_local.
+		// Since we don't have access to q_ref in SkeletalMeshObject, we convert simply matrices to quaternions.
+		for (int32 QuatIdx = 0; QuatIdx < NumQuats; ++QuatIdx)
+		{
+			OutRefToLocalQuats[QuatIdx] = RefToLocalMatrices[QuatIdx].GetMatrixWithoutScale().ToQuat();
+		}
+	}
+}	// namespace UE::MLDeformer
diff --git a/Engine/Plugins/Animation/MLDeformer/MLDeformerFramework/Source/MLDeformerFramework/Private/DeformerGraphDataInterfaceSkeletonWithQuats.h b/Engine/Plugins/Animation/MLDeformer/MLDeformerFramework/Source/MLDeformerFramework/Private/DeformerGraphDataInterfaceSkeletonWithQuats.h
new file mode 100644
--- /dev/null
+++ b/Engine/Plugins/Animation/MLDeformer/MLDeformerFramework/Source/MLDeformerFramework/Private/DeformerGraphDataInterfaceSkeletonWithQuats.h
@@ -0,0 +1,83 @@
+// Copyright Epic Games, Inc. All Rights Reserved.
+
+#pragma once
+
+#include "OptimusComputeDataInterface.h"
+#include "ComputeFramework/ComputeDataProvider.h"
+#include "RenderGraphFwd.h"
+#include "DeformerGraphDataInterfaceSkeletonWithQuats.generated.h"
+
+class FSkeletalMeshObject;
+class FSkeletonWithQuatsDataInterfaceParameters;
+class USkinnedMeshComponent;
+
+/** Compute Framework Data Interface for skeletal data. */
+UCLASS(Category = ComputeFramework)
+class MLDEFORMERFRAMEWORK_API UDeformerGraphSkeletonWithWeightedQuatsDataInterface
+	: public UOptimusComputeDataInterface
+{
+	GENERATED_BODY()
+
+public:
+	//~ Begin UOptimusComputeDataInterface Interface
+	FString GetDisplayName() const override;
+	TArray<FOptimusCDIPinDefinition> GetPinDefinitions() const override;
+	TSubclassOf<UActorComponent> GetRequiredComponentClass() const override;
+	//~ End UOptimusComputeDataInterface Interface
+	
+	//~ Begin UComputeDataInterface Interface
+	TCHAR const* GetClassName() const override { return TEXT("MLDeformer_SkeletonWithQuats"); }
+	void GetSupportedInputs(TArray<FShaderFunctionDefinition>& OutFunctions) const override;
+	void GetShaderParameters(TCHAR const* UID, FShaderParametersMetadataBuilder& InOutBuilder, FShaderParametersMetadataAllocations& InOutAllocations) const override;
+	void GetPermutations(FComputeKernelPermutationVector& OutPermutationVector) const override;
+	TCHAR const* GetShaderVirtualPath() const override;
+	void GetShaderHash(FString& InOutKey) const override;
+	void GetHLSL(FString& OutHLSL, FString const& InDataInterfaceName) const override;
+	UComputeDataProvider* CreateDataProvider(TObjectPtr<UObject> InBinding, uint64 InInputMask, uint64 InOutputMask) const override;
+	//~ End UComputeDataInterface Interface
+
+private:
+	static TCHAR const* TemplateFilePath;
+};
+
+/** Compute Framework Data Provider for reading skeletal mesh. */
+UCLASS(BlueprintType, EditInlineNew, Category = ComputeFramework)
+class MLDEFORMERFRAMEWORK_API UDeformerGraphSkeletonWithWeightedQuatsDataProvider
+	: public UComputeDataProvider
+{
+	GENERATED_BODY()
+
+public:
+	UPROPERTY(EditAnywhere, BlueprintReadWrite, Category = Binding)
+	TObjectPtr<USkinnedMeshComponent> SkinnedMesh = nullptr;
+
+	//~ Begin UComputeDataProvider Interface
+	FComputeDataProviderRenderProxy* GetRenderProxy() override;
+	//~ End UComputeDataProvider Interface
+};
+
+namespace UE::MLDeformer
+{
+	class MLDEFORMERFRAMEWORK_API FDeformerGraphSkeletonWithWeightedQuatsDataProviderProxy
+		: public FComputeDataProviderRenderProxy
+	{
+	public:
+		FDeformerGraphSkeletonWithWeightedQuatsDataProviderProxy(USkinnedMeshComponent* SkinnedMeshComponent);
+
+		//~ Begin FComputeDataProviderRenderProxy Interface
+		bool IsValid(FValidationData const& InValidationData) const override;
+		void GatherPermutations(FPermutationData& InOutPermutationData) const override;
+		void GatherDispatchData(FDispatchData const& InDispatchData) override;
+		void AllocateResources(FRDGBuilder& GraphBuilder) override;
+		//~ End FComputeDataProviderRenderProxy Interface
+
+	private:
+		using FParameters = FSkeletonWithQuatsDataInterfaceParameters;
+		void CacheRefToLocalQuats(TArray<FQuat4f>& OutRefToLocalQuats) const;
+
+		FSkeletalMeshObject* SkeletalMeshObject = nullptr;
+		uint32 BoneRevisionNumber = 0;
+		TArray<TArray<FQuat4f> > PerSectionRefToLocalQuats;
+		TArray<FRDGBufferSRVRef> PerSectionRefToLocalQuatsSRVs;
+	};
+}	// namespace UE::MLDeformer
diff --git a/Engine/Plugins/Animation/MLDeformer/MLDeformerFramework/Source/MLDeformerFramework/Private/MLDeformerModelInstance.cpp b/Engine/Plugins/Animation/MLDeformer/MLDeformerFramework/Source/MLDeformerFramework/Private/MLDeformerModelInstance.cpp
--- a/Engine/Plugins/Animation/MLDeformer/MLDeformerFramework/Source/MLDeformerFramework/Private/MLDeformerModelInstance.cpp
+++ b/Engine/Plugins/Animation/MLDeformer/MLDeformerFramework/Source/MLDeformerFramework/Private/MLDeformerModelInstance.cpp
@@ -138,7 +138,7 @@ FString UMLDeformerModelInstance::CheckCompatibility(USkeletalMeshComponent* InS
 	// If we're not compatible, generate a compatibility string.
 	USkeletalMesh* SkelMesh = InSkelMeshComponent ? InSkelMeshComponent->GetSkeletalMeshAsset() : nullptr;
 	UMLDeformerInputInfo* InputInfo = Model->GetInputInfo();
-	if (SkelMesh && !InputInfo->IsCompatible(SkelMesh) && Model->GetDeformerAsset())
+	if (SkelMesh && !InputInfo->IsCompatible(SkelMesh) && Model->GetDeformerAsset() && Model->IsTrained())
 	{
 		ErrorText += InputInfo->GenerateCompatibilityErrorString(SkelMesh);
 		if (bLogIssues)
diff --git a/Engine/Plugins/Animation/MLDeformer/MLDeformerFramework/Source/MLDeformerFramework/Private/MLDeformerMorphModel.cpp b/Engine/Plugins/Animation/MLDeformer/MLDeformerFramework/Source/MLDeformerFramework/Private/MLDeformerMorphModel.cpp
--- a/Engine/Plugins/Animation/MLDeformer/MLDeformerFramework/Source/MLDeformerFramework/Private/MLDeformerMorphModel.cpp
+++ b/Engine/Plugins/Animation/MLDeformer/MLDeformerFramework/Source/MLDeformerFramework/Private/MLDeformerMorphModel.cpp
@@ -191,9 +191,12 @@ int32 UMLDeformerMorphModel::GetNumMorphTargets(int32 LOD) const
 
 bool UMLDeformerMorphModel::CanDynamicallyUpdateMorphTargets() const
 {
-	const int32 LOD = 0;
-	const bool bResult = (GetMorphTargetDeltas().Num() == (GetNumBaseMeshVerts() * GetNumMorphTargets(LOD))) || bHasOnlyEmptyMorphs;
-	return bResult;
+	if (GetNumBaseMeshVerts() > 0)
+	{
+		return (GetMorphTargetDeltas().Num() % GetNumBaseMeshVerts() == 0) || bHasOnlyEmptyMorphs;
+	}
+
+	return false;
 }
 
 UMLDeformerModelInstance* UMLDeformerMorphModel::CreateModelInstance(UMLDeformerComponent* Component)
diff --git a/Engine/Plugins/Animation/MLDeformer/MLDeformerFramework/Source/MLDeformerFramework/Public/MLDeformerGeomCacheVizSettings.h b/Engine/Plugins/Animation/MLDeformer/MLDeformerFramework/Source/MLDeformerFramework/Public/MLDeformerGeomCacheVizSettings.h
--- a/Engine/Plugins/Animation/MLDeformer/MLDeformerFramework/Source/MLDeformerFramework/Public/MLDeformerGeomCacheVizSettings.h
+++ b/Engine/Plugins/Animation/MLDeformer/MLDeformerFramework/Source/MLDeformerFramework/Public/MLDeformerGeomCacheVizSettings.h
@@ -10,7 +10,7 @@
 #include "MLDeformerGeomCacheVizSettings.generated.h"
 
 /**
- * The vizualization settings for a model that has a geometry cache.
+ * The visualization settings for a model that has a geometry cache.
  * This can be used in combination with a UMLDeformerGeomCacheModel.
  */
 UCLASS()
diff --git a/Engine/Plugins/Animation/MLDeformer/MLDeformerFramework/Source/MLDeformerFramework/Public/MLDeformerModel.h b/Engine/Plugins/Animation/MLDeformer/MLDeformerFramework/Source/MLDeformerFramework/Public/MLDeformerModel.h
--- a/Engine/Plugins/Animation/MLDeformer/MLDeformerFramework/Source/MLDeformerFramework/Public/MLDeformerModel.h
+++ b/Engine/Plugins/Animation/MLDeformer/MLDeformerFramework/Source/MLDeformerFramework/Public/MLDeformerModel.h
@@ -44,6 +44,17 @@ enum class EMLDeformerMaskChannel : uint8
 	VertexAttribute
 };
 
+/** The skinning mode to use as base. */
+UENUM()
+enum class EMLDeformerSkinningMode : int8
+{
+	/** Linear blend skinning. This is the fastest at runtime, but can have a harder time reconstructing the ground truth. */
+	Linear,
+
+	/** Dual quaternion skinning. This is slower at runtime, but can result in better deformations. */
+	DualQuaternion
+};
+
 
 namespace UE::MLDeformer
 {
diff --git a/Engine/Plugins/Animation/MLDeformer/MLDeformerFramework/Source/MLDeformerFramework/Public/MLDeformerMorphModel.h b/Engine/Plugins/Animation/MLDeformer/MLDeformerFramework/Source/MLDeformerFramework/Public/MLDeformerMorphModel.h
--- a/Engine/Plugins/Animation/MLDeformer/MLDeformerFramework/Source/MLDeformerFramework/Public/MLDeformerMorphModel.h
+++ b/Engine/Plugins/Animation/MLDeformer/MLDeformerFramework/Source/MLDeformerFramework/Public/MLDeformerMorphModel.h
@@ -223,6 +223,7 @@ public:
 	 * So the total number of items in the array returned equals (NumMorphTargets * NumBaseMeshVerts).
 	 */
 	const TArray<FVector3f>& GetMorphTargetDeltas() const				{ return MorphTargetDeltas; }
+	TArray<FVector3f>& GetMorphTargetDeltas()							{ return MorphTargetDeltas; }
 
 	/**
 	 * Get the morph target set.
diff --git a/Engine/Plugins/Animation/MLDeformer/MLDeformerFramework/Source/MLDeformerFramework/Public/MLDeformerMorphModelVizSettings.h b/Engine/Plugins/Animation/MLDeformer/MLDeformerFramework/Source/MLDeformerFramework/Public/MLDeformerMorphModelVizSettings.h
--- a/Engine/Plugins/Animation/MLDeformer/MLDeformerFramework/Source/MLDeformerFramework/Public/MLDeformerMorphModelVizSettings.h
+++ b/Engine/Plugins/Animation/MLDeformer/MLDeformerFramework/Source/MLDeformerFramework/Public/MLDeformerMorphModelVizSettings.h
@@ -7,7 +7,7 @@
 #include "MLDeformerMorphModelVizSettings.generated.h"
 
 /**
- * The vizualization settings specific to the UMLDeformerMorphModel class, or inherited classes.
+ * The visualization settings specific to the UMLDeformerMorphModel class, or inherited classes.
  */
 UCLASS()
 class MLDEFORMERFRAMEWORK_API UMLDeformerMorphModelVizSettings
diff --git a/Engine/Plugins/Animation/MLDeformer/MLDeformerFramework/Source/MLDeformerFramework/Public/MLDeformerVizSettings.h b/Engine/Plugins/Animation/MLDeformer/MLDeformerFramework/Source/MLDeformerFramework/Public/MLDeformerVizSettings.h
--- a/Engine/Plugins/Animation/MLDeformer/MLDeformerFramework/Source/MLDeformerFramework/Public/MLDeformerVizSettings.h
+++ b/Engine/Plugins/Animation/MLDeformer/MLDeformerFramework/Source/MLDeformerFramework/Public/MLDeformerVizSettings.h
@@ -49,7 +49,7 @@ struct MLDEFORMERFRAMEWORK_API FMLDeformerCompareActor
 };
 
 /**
- * The vizualization settings.
+ * The visualization settings.
  */
 UCLASS()
 class MLDEFORMERFRAMEWORK_API UMLDeformerVizSettings
diff --git a/Engine/Plugins/Animation/MLDeformer/MLDeformerFramework/Source/MLDeformerFrameworkEditor/Private/MLDeformerEditorModel.cpp b/Engine/Plugins/Animation/MLDeformer/MLDeformerFramework/Source/MLDeformerFrameworkEditor/Private/MLDeformerEditorModel.cpp
--- a/Engine/Plugins/Animation/MLDeformer/MLDeformerFramework/Source/MLDeformerFrameworkEditor/Private/MLDeformerEditorModel.cpp
+++ b/Engine/Plugins/Animation/MLDeformer/MLDeformerFramework/Source/MLDeformerFrameworkEditor/Private/MLDeformerEditorModel.cpp
@@ -1162,6 +1162,8 @@ namespace UE::MLDeformer
 			return;
 		}
 
+		const bool bPropertyValueSetOrResetToDefault = (PropertyChangedEvent.ChangeType & EPropertyChangeType::ValueSet) || (PropertyChangedEvent.ChangeType & EPropertyChangeType::ResetToDefault);
+
 		// When we change one of these properties below, restart animations etc.
 		if (Property->GetFName() == UMLDeformerModel::GetSkeletalMeshPropertyName())
 		{
@@ -1219,7 +1221,7 @@ namespace UE::MLDeformer
 		}
 		else if (Property->GetFName() == GET_MEMBER_NAME_CHECKED(FMLDeformerCompareActor, DeformerAsset))
 		{
-			if (PropertyChangedEvent.ChangeType == EPropertyChangeType::ValueSet)
+			if (bPropertyValueSetOrResetToDefault)
 			{
 				TriggerInputAssetChanged(true);
 				UpdateActorVisibility();
@@ -1228,7 +1230,7 @@ namespace UE::MLDeformer
 		}
 		else if (Property->GetFName() == GET_MEMBER_NAME_CHECKED(FMLDeformerCompareActor, Name))
 		{
-			if (PropertyChangedEvent.ChangeType == EPropertyChangeType::ValueSet)
+			if (bPropertyValueSetOrResetToDefault)
 			{
 				UpdateCompareActorLabels();
 			}
@@ -1236,7 +1238,7 @@ namespace UE::MLDeformer
 		else
 		if (Property->GetFName() == UMLDeformerModel::GetAlignmentTransformPropertyName())
 		{
-			if (PropertyChangedEvent.ChangeType == EPropertyChangeType::ValueSet)
+			if (bPropertyValueSetOrResetToDefault)
 			{
 				SetResamplingInputOutputsNeeded(true);
 				SampleDeltas();
@@ -1257,7 +1259,7 @@ namespace UE::MLDeformer
 		}
 		else if (Property->GetFName() == UMLDeformerModel::GetMaxNumLODsPropertyName())
 		{
-			if (PropertyChangedEvent.ChangeType == EPropertyChangeType::ValueSet)
+			if (bPropertyValueSetOrResetToDefault)
 			{
 				OnMaxNumLODsChanged();
 			}
@@ -1903,7 +1905,7 @@ namespace UE::MLDeformer
 		{
 			const TArray<float>& VertexDeltas = Sampler->GetVertexDeltas();
 			const TArray<FVector3f>& LinearSkinnedPositions = Sampler->GetSkinnedVertexPositions();
-			bool bDrawDeltas = (VizSettings->GetDrawVertexDeltas() && (VertexDeltas.Num() / 3) == LinearSkinnedPositions.Num());
+			bool bDrawDeltas = VizSettings->GetDrawVertexDeltas();
 
 			// Disable drawing deltas when we're playing the training anim sequence, as that can get too slow.
 			const FMLDeformerEditorActor* BaseActor = FindEditorActor(ActorID_Train_Base);
@@ -1912,8 +1914,15 @@ namespace UE::MLDeformer
 				bDrawDeltas &= !BaseActor->IsPlaying();
 			}
 
-			// Draw the deltas for the current frame.
 			USkeletalMesh* SkelMesh = Model->GetSkeletalMesh();
+			if (bDrawDeltas && Sampler->GetVertexDeltas().IsEmpty() && Model->GetSkeletalMesh() && SkelMesh->GetImportedModel())
+			{
+				SampleDeltas();
+			}
+
+			bDrawDeltas &= ((VertexDeltas.Num() / 3) == LinearSkinnedPositions.Num());
+
+			// Draw the deltas for the current frame.
 			if (bDrawDeltas && Model->GetSkeletalMesh() && SkelMesh->GetImportedModel() && !IsScrubbingTimeline())
 			{
 				const FLinearColor DeltasColor = FMLDeformerEditorStyle::Get().GetColor("MLDeformer.Deltas.Color");
@@ -2109,8 +2118,26 @@ namespace UE::MLDeformer
 	void FMLDeformerEditorModel::CreateHeatMapDeformerGraph()
 	{
 		const FString HeatMapDeformerPath = GetHeatMapDeformerGraphPath();
-		UObject* DeformerObject = StaticLoadObject(UMeshDeformer::StaticClass(), nullptr, *HeatMapDeformerPath);
-		HeatMapDeformerGraph = Cast<UMeshDeformer>(DeformerObject);
+		if (!HeatMapDeformerPath.IsEmpty())
+		{
+			UObject* DeformerObject = StaticLoadObject(UMeshDeformer::StaticClass(), nullptr, *HeatMapDeformerPath);
+			HeatMapDeformerGraph = Cast<UMeshDeformer>(DeformerObject);
+		}
+		else
+		{
+			HeatMapDeformerGraph = nullptr;
+		}
+
+		const FString HeatMapDeformerDualQuatPath = GetHeatMapDeformerGraphDualQuatPath();
+		if (!HeatMapDeformerDualQuatPath.IsEmpty())
+		{
+			UObject* DeformerObjectDualQuat = StaticLoadObject(UMeshDeformer::StaticClass(), nullptr, *HeatMapDeformerDualQuatPath);
+			HeatMapDeformerGraphDualQuat = Cast<UMeshDeformer>(DeformerObjectDualQuat);
+		}
+		else
+		{
+			HeatMapDeformerGraphDualQuat = nullptr;
+		}
 	}
 
 	void FMLDeformerEditorModel::CreateHeatMapAssets()
@@ -2119,6 +2146,11 @@ namespace UE::MLDeformer
 		CreateHeatMapDeformerGraph();
 	}
 
+	UMeshDeformer* FMLDeformerEditorModel::GetActiveHeatMapDeformer() const
+	{
+		return HeatMapDeformerGraph;
+	}
+
 	void FMLDeformerEditorModel::SetHeatMapMaterialEnabled(bool bEnabled)
 	{
 		FMLDeformerEditorActor* EditorActor = FindEditorActor(ActorID_Test_MLDeformed);
@@ -2209,7 +2241,7 @@ namespace UE::MLDeformer
 			{
 				UMeshDeformer* MeshDeformer = IsTrained() ? VizSettings->GetDeformerGraph() : nullptr;	
 				const bool bUseHeatMapDeformer = VizSettings->GetShowHeatMap();
-				SkelMeshComponent->SetMeshDeformer(bUseHeatMapDeformer ? HeatMapDeformerGraph.Get() : MeshDeformer);
+				SkelMeshComponent->SetMeshDeformer(bUseHeatMapDeformer ? GetActiveHeatMapDeformer() : MeshDeformer);
 			}
 		}
 
@@ -2234,6 +2266,10 @@ namespace UE::MLDeformer
 					{
 						UMeshDeformer* MeshDeformer = Asset->GetModel()->GetVizSettings()->GetDeformerGraph();	
 						const bool bUseHeatMapDeformer = VizSettings->GetShowHeatMap();
+
+						// NOTE: Even though we use a HeatMapDeformerGraph here, it is not entirely correct, as if we want to be fully correct
+						// we should be using the active heat map deformer from an editor model that fits the compare model.
+						// So right now we assume a linear skinned heatmap, but if the compare model uses dual quaternion skinning then it's heatmap would appear incorrect.
 						SkelMeshComponent->SetMeshDeformer(bUseHeatMapDeformer ? HeatMapDeformerGraph.Get() : MeshDeformer);
 					}
 				}
@@ -2689,6 +2725,11 @@ namespace UE::MLDeformer
 		return FString(TEXT("/MLDeformerFramework/Deformers/DG_MLDeformerModel_HeatMap.DG_MLDeformerModel_HeatMap"));
 	}
 
+	FString FMLDeformerEditorModel::GetHeatMapDeformerGraphDualQuatPath() const
+	{
+		return FString();
+	}
+
 	void FMLDeformerEditorModel::RefreshInputWidget()
 	{
 		if (InputWidget.IsValid())
diff --git a/Engine/Plugins/Animation/MLDeformer/MLDeformerFramework/Source/MLDeformerFrameworkEditor/Private/MLDeformerGeomCacheEditorModel.cpp b/Engine/Plugins/Animation/MLDeformer/MLDeformerFramework/Source/MLDeformerFrameworkEditor/Private/MLDeformerGeomCacheEditorModel.cpp
--- a/Engine/Plugins/Animation/MLDeformer/MLDeformerFramework/Source/MLDeformerFrameworkEditor/Private/MLDeformerGeomCacheEditorModel.cpp
+++ b/Engine/Plugins/Animation/MLDeformer/MLDeformerFramework/Source/MLDeformerFrameworkEditor/Private/MLDeformerGeomCacheEditorModel.cpp
@@ -252,8 +252,7 @@ namespace UE::MLDeformer
 		FMLDeformerEditorModel::OnPropertyChanged(PropertyChangedEvent);
 
 		// Start frame changed.
-		if (PropertyChangedEvent.GetMemberPropertyName() == TEXT("TrainingInputAnims") && 
-			Property->GetFName() == FMLDeformerTrainingInputAnim::GetStartFramePropertyName())
+		if (PropertyChangedEvent.GetMemberPropertyName() == TEXT("TrainingInputAnims") && Property->GetFName() == FMLDeformerTrainingInputAnim::GetStartFramePropertyName())
 		{
 			for (FMLDeformerGeomCacheTrainingInputAnim& InputAnim : GetGeomCacheModel()->GetTrainingInputAnims())
 			{
@@ -265,9 +264,7 @@ namespace UE::MLDeformer
 			SetResamplingInputOutputsNeeded(true);
 			TriggerInputAssetChanged(true);
 		}
-		else // End Frame changed.
-		if (PropertyChangedEvent.GetMemberPropertyName() == TEXT("TrainingInputAnims") && 
-		Property->GetFName() == FMLDeformerTrainingInputAnim::GetEndFramePropertyName())
+		else if (PropertyChangedEvent.GetMemberPropertyName() == TEXT("TrainingInputAnims") && Property->GetFName() == FMLDeformerTrainingInputAnim::GetEndFramePropertyName())
 		{
 			for (FMLDeformerGeomCacheTrainingInputAnim& InputAnim : GetGeomCacheModel()->GetTrainingInputAnims())
 			{
@@ -279,9 +276,7 @@ namespace UE::MLDeformer
 			SetResamplingInputOutputsNeeded(true);
 			TriggerInputAssetChanged(true);
 		}
-		else
-		if (Property->GetFName() == UMLDeformerGeomCacheModel::GetTrainingInputAnimsPropertyName() ||
-		    PropertyChangedEvent.GetMemberPropertyName() == TEXT("TrainingInputAnims"))
+		else if (Property->GetFName() == UMLDeformerGeomCacheModel::GetTrainingInputAnimsPropertyName() || PropertyChangedEvent.GetMemberPropertyName() == TEXT("TrainingInputAnims"))
 		{
 			if (GetEditor() && GetEditor()->GetTimeSlider())
 			{
diff --git a/Engine/Plugins/Animation/MLDeformer/MLDeformerFramework/Source/MLDeformerFrameworkEditor/Private/MLDeformerGeomCacheSampler.cpp b/Engine/Plugins/Animation/MLDeformer/MLDeformerFramework/Source/MLDeformerFrameworkEditor/Private/MLDeformerGeomCacheSampler.cpp
--- a/Engine/Plugins/Animation/MLDeformer/MLDeformerFramework/Source/MLDeformerFrameworkEditor/Private/MLDeformerGeomCacheSampler.cpp
+++ b/Engine/Plugins/Animation/MLDeformer/MLDeformerFramework/Source/MLDeformerFrameworkEditor/Private/MLDeformerGeomCacheSampler.cpp
@@ -129,15 +129,26 @@ namespace UE::MLDeformer
 							FVector3f Delta = FVector3f::ZeroVector;
 							if (VertexDeltaSpace == EVertexDeltaSpace::PreSkinning)
 							{
-								// Calculate the inverse skinning transform for this vertex.
-								const FMatrix44f InvSkinningTransform = CalcInverseSkinningTransform(RenderVertexIndex, SkelMeshLODData, SkinWeightBuffer);
-
-								// Calculate the pre-skinning data.
-								const FSkeletalMeshLODRenderData& LODData = SkeletalMesh->GetResourceForRendering()->LODRenderData[0];
-								const FVector3f UnskinnedPosition = LODData.StaticVertexBuffers.PositionVertexBuffer.VertexPosition(RenderVertexIndex);
-								const FVector3f GeomCacheVertexPos = (FVector3f)AlignmentTransform.TransformPosition((FVector)GeomCacheMeshData.Positions[GeomCacheVertexIndex]);
-								const FVector3f PreSkinningTargetPos = InvSkinningTransform.TransformPosition(GeomCacheVertexPos);
-								Delta = PreSkinningTargetPos - UnskinnedPosition;
+								if (SkinningMode == EMLDeformerSkinningMode::Linear)
+								{
+									// Calculate the inverse skinning transform for this vertex.
+									const FMatrix44f InvSkinningTransform = CalcInverseSkinningTransform(RenderVertexIndex, SkelMeshLODData, SkinWeightBuffer);
+
+									// Calculate the pre-skinning data.
+									const FSkeletalMeshLODRenderData& LODData = SkeletalMesh->GetResourceForRendering()->LODRenderData[0];
+									const FVector3f UnskinnedPosition = LODData.StaticVertexBuffers.PositionVertexBuffer.VertexPosition(RenderVertexIndex);
+									const FVector3f GeomCacheVertexPos = (FVector3f)AlignmentTransform.TransformPosition((FVector)GeomCacheMeshData.Positions[GeomCacheVertexIndex]);
+									const FVector3f PreSkinningTargetPos = InvSkinningTransform.TransformPosition(GeomCacheVertexPos);
+									Delta = PreSkinningTargetPos - UnskinnedPosition;
+								}
+								else // Dual Quaternion Skinning.
+								{
+									const FVector3f SkinnedVertexPos = SkinnedVertexPositions[SkinnedVertexIndex];
+									const FVector3f GeomCacheVertexPos = (FVector3f)AlignmentTransform.TransformPosition((FVector)GeomCacheMeshData.Positions[GeomCacheVertexIndex]);
+									const FVector3f WorldDelta = GeomCacheVertexPos - SkinnedVertexPos;
+									const FSkeletalMeshLODRenderData& LODData = SkeletalMesh->GetResourceForRendering()->LODRenderData[0];
+									Delta = CalcDualQuaternionDelta(RenderVertexIndex, WorldDelta, LODData, SkinWeightBuffer);
+								}
 							}
 							else if (VertexDeltaSpace == EVertexDeltaSpace::PostSkinning)
 							{
diff --git a/Engine/Plugins/Animation/MLDeformer/MLDeformerFramework/Source/MLDeformerFrameworkEditor/Private/MLDeformerModelDetails.cpp b/Engine/Plugins/Animation/MLDeformer/MLDeformerFramework/Source/MLDeformerFrameworkEditor/Private/MLDeformerModelDetails.cpp
--- a/Engine/Plugins/Animation/MLDeformer/MLDeformerFramework/Source/MLDeformerFrameworkEditor/Private/MLDeformerModelDetails.cpp
+++ b/Engine/Plugins/Animation/MLDeformer/MLDeformerFramework/Source/MLDeformerFrameworkEditor/Private/MLDeformerModelDetails.cpp
@@ -58,7 +58,7 @@ namespace UE::MLDeformer
 		BaseMeshCategoryBuilder = &DetailLayoutBuilder->EditCategory("Base Mesh", FText::GetEmpty(), ECategoryPriority::Important);
 		InputOutputCategoryBuilder = &DetailLayoutBuilder->EditCategory("Inputs", FText::GetEmpty(), ECategoryPriority::Important);
 		TrainingSettingsCategoryBuilder = &DetailLayoutBuilder->EditCategory("Training Settings", FText::GetEmpty(), ECategoryPriority::Important);
-		LODSettingsCategoryBuilder = &DetailLayoutBuilder->EditCategory("LOD Generation Settings", FText::GetEmpty(), ECategoryPriority::Important);
+		LODSettingsCategoryBuilder = &DetailLayoutBuilder->EditCategory("LOD Generation Settings", FText::GetEmpty(), ECategoryPriority::Uncommon);
 		LODSettingsCategoryBuilder->SetCategoryVisibility(Model->DoesSupportLOD());
 	}
 
@@ -152,7 +152,7 @@ namespace UE::MLDeformer
 		// Create the inputs widget.
 		TSharedPtr<SMLDeformerInputWidget> InputWidget = EditorModel->CreateInputWidget();
 		EditorModel->SetInputWidget(InputWidget);
-		IDetailGroup& NetworkInputsGroup = InputOutputCategoryBuilder->AddGroup("NetworkInputs", LOCTEXT("NetworkInputs", "Network Inputs"), false, true);
+		IDetailGroup& NetworkInputsGroup = InputOutputCategoryBuilder->AddGroup("NetworkInputs", LOCTEXT("NetworkInputs", "Network Inputs"), false, false);
 		NetworkInputsGroup.AddWidgetRow().WholeRowContent().Widget = InputWidget.ToSharedRef();
 
 		AddBoneInputErrors();
diff --git a/Engine/Plugins/Animation/MLDeformer/MLDeformerFramework/Source/MLDeformerFrameworkEditor/Private/MLDeformerMorphModelDetails.cpp b/Engine/Plugins/Animation/MLDeformer/MLDeformerFramework/Source/MLDeformerFrameworkEditor/Private/MLDeformerMorphModelDetails.cpp
--- a/Engine/Plugins/Animation/MLDeformer/MLDeformerFramework/Source/MLDeformerFrameworkEditor/Private/MLDeformerMorphModelDetails.cpp
+++ b/Engine/Plugins/Animation/MLDeformer/MLDeformerFramework/Source/MLDeformerFrameworkEditor/Private/MLDeformerMorphModelDetails.cpp
@@ -37,7 +37,7 @@ namespace UE::MLDeformer
 	void FMLDeformerMorphModelDetails::CreateCategories()
 	{
 		FMLDeformerGeomCacheModelDetails::CreateCategories();
-		MorphTargetCategoryBuilder = &DetailLayoutBuilder->EditCategory("Morph Targets", FText::GetEmpty(), ECategoryPriority::Important);
+		MorphTargetCategoryBuilder = &DetailLayoutBuilder->EditCategory("Morph Targets", FText::GetEmpty(), ECategoryPriority::Default);
 	}
 
 	bool FMLDeformerMorphModelDetails::ShouldShowShadingError() const
@@ -174,7 +174,7 @@ namespace UE::MLDeformer
 
 		MorphTargetCategoryBuilder->AddProperty(DetailLayoutBuilder->GetProperty(UMLDeformerMorphModel::GetIncludeMorphTargetNormalsPropertyName(), UMLDeformerMorphModel::StaticClass()));
 
-		IDetailGroup& CompressionGroup = MorphTargetCategoryBuilder->AddGroup("Compression", LOCTEXT("MorphCompressionGroupLabel", "Compression"), false, true);
+		IDetailGroup& CompressionGroup = MorphTargetCategoryBuilder->AddGroup("Compression", LOCTEXT("MorphCompressionGroupLabel", "Compression"), false, false);
 		CompressionGroup.AddPropertyRow(DetailLayoutBuilder->GetProperty(UMLDeformerMorphModel::GetMorphDeltaZeroThresholdPropertyName(), UMLDeformerMorphModel::StaticClass()));
 		CompressionGroup.AddPropertyRow(DetailLayoutBuilder->GetProperty(UMLDeformerMorphModel::GetMorphCompressionLevelPropertyName(), UMLDeformerMorphModel::StaticClass()));
 
diff --git a/Engine/Plugins/Animation/MLDeformer/MLDeformerFramework/Source/MLDeformerFrameworkEditor/Private/MLDeformerMorphModelEditorModel.cpp b/Engine/Plugins/Animation/MLDeformer/MLDeformerFramework/Source/MLDeformerFrameworkEditor/Private/MLDeformerMorphModelEditorModel.cpp
--- a/Engine/Plugins/Animation/MLDeformer/MLDeformerFramework/Source/MLDeformerFrameworkEditor/Private/MLDeformerMorphModelEditorModel.cpp
+++ b/Engine/Plugins/Animation/MLDeformer/MLDeformerFramework/Source/MLDeformerFrameworkEditor/Private/MLDeformerMorphModelEditorModel.cpp
@@ -85,7 +85,7 @@ namespace UE::MLDeformer
 			Property->GetFName() == UMLDeformerMorphModel::GetGlobalMaskAttributePropertyName() ||
 			Property->GetFName() == UMLDeformerMorphModel::GetSkeletalMeshPropertyName())
 		{
-			if (PropertyChangedEvent.ChangeType == EPropertyChangeType::ValueSet)
+			if (PropertyChangedEvent.ChangeType & EPropertyChangeType::ValueSet || PropertyChangedEvent.ChangeType & EPropertyChangeType::ResetToDefault)
 			{
 				if (GetMorphModel()->CanDynamicallyUpdateMorphTargets())
 				{
@@ -129,6 +129,11 @@ namespace UE::MLDeformer
 		return FString(TEXT("/MLDeformerFramework/Deformers/DG_MLDeformerModel_GPUMorph_HeatMap.DG_MLDeformerModel_GPUMorph_HeatMap"));
 	}
 
+	FString FMLDeformerMorphModelEditorModel::GetHeatMapDeformerGraphDualQuatPath() const
+	{
+		return FString(TEXT("/MLDeformerFramework/Deformers/DG_MLDeformerModel_GPUMorph_HeatMap_DQ.DG_MLDeformerModel_GPUMorph_HeatMap_DQ"));
+	}
+
 	void FMLDeformerMorphModelEditorModel::OnPreTraining()
 	{
 		// Backup the morph target deltas in case we abort training.
@@ -781,7 +786,7 @@ namespace UE::MLDeformer
 	void FMLDeformerMorphModelEditorModel::InitEngineMorphTargets(const TArray<FVector3f>& Deltas)
 	{
 		UMLDeformerMorphModel* MorphModel = GetMorphModel();
-		if (MorphModel->GetMorphTargetDeltas().IsEmpty())
+		if (Deltas.IsEmpty())
 		{
 			return;
 		}
@@ -790,7 +795,7 @@ namespace UE::MLDeformer
 		TArray<FVector3f> MorphTargetDeltas = Deltas;
 		ZeroDeltasByLengthThreshold(MorphTargetDeltas, MorphModel->GetMorphDeltaZeroThreshold());
 
-		// Turn the delta buffer in a set of engine morph targets.
+		// Turn the delta buffer into a set of engine morph targets.
 		const int32 LODZero = 0;
 		const bool bIncludeNormals = MorphModel->GetIncludeMorphTargetNormals();
 		const EMLDeformerMaskChannel MaskChannel = MorphModel->GetMaskChannel();
diff --git a/Engine/Plugins/Animation/MLDeformer/MLDeformerFramework/Source/MLDeformerFrameworkEditor/Private/MLDeformerSampler.cpp b/Engine/Plugins/Animation/MLDeformer/MLDeformerFramework/Source/MLDeformerFrameworkEditor/Private/MLDeformerSampler.cpp
--- a/Engine/Plugins/Animation/MLDeformer/MLDeformerFramework/Source/MLDeformerFrameworkEditor/Private/MLDeformerSampler.cpp
+++ b/Engine/Plugins/Animation/MLDeformer/MLDeformerFramework/Source/MLDeformerFrameworkEditor/Private/MLDeformerSampler.cpp
@@ -123,7 +123,7 @@ namespace UE::MLDeformer
 		if (SkeletalMeshComponent && SkeletalMesh)
 		{
 			UpdateSkeletalMeshComponent();
-			if (VertexDeltaSpace == EVertexDeltaSpace::PostSkinning)
+			if (VertexDeltaSpace == EVertexDeltaSpace::PostSkinning || SkinningMode == EMLDeformerSkinningMode::DualQuaternion)
 			{
 				UpdateSkinnedPositions();
 			}
@@ -175,6 +175,91 @@ namespace UE::MLDeformer
 		return InvSkinningTransform.Inverse();
 	}
 
+	namespace
+	{
+		template<typename QuatType>
+		QuatType Conjugate(const QuatType& Q)
+		{
+			return QuatType(-Q.X, -Q.Y, -Q.Z, Q.W);
+		}
+		
+		template<typename QuatType>
+		float Inner(const QuatType& Q1, const QuatType& Q2)
+		{
+			return Q1.X * Q2.X + Q1.Y * Q2.Y + Q1.Z * Q2.Z + Q1.W * Q2.W;
+		}
+		
+		template<typename QuatType>
+		QuatType FromVector(const FVector3f& V)
+		{
+			return QuatType(V[0], V[1], V[2], 0);
+		}
+		
+		template<typename QuatType>
+		FVector3f ToVector(const QuatType& Q)
+		{
+			return FVector3f(Q.X, Q.Y, Q.Z);
+		}
+	};
+
+	FVector3f FMLDeformerSampler::CalcDualQuaternionDelta(int32 VertexIndex, const FVector3f& WorldDelta, const FSkeletalMeshLODRenderData& SkelMeshLODData, const FSkinWeightVertexBuffer& SkinWeightBuffer) const
+	{
+		check(SkeletalMeshComponent);
+		const USkeletalMesh* Mesh = SkeletalMeshComponent->GetSkeletalMeshAsset();
+		check(Mesh);
+		const FSkeletalMeshRenderData* RenderData = Mesh->GetResourceForRendering();
+		constexpr int32 LODIndex = 0;
+		check(RenderData && RenderData->LODRenderData.IsValidIndex(LODIndex));
+		const FSkeletalMeshLODRenderData& LODData = RenderData->LODRenderData[LODIndex];
+
+		// Find the render section, which we need to find the right bone index.
+		int32 SectionIndex = INDEX_NONE;
+		int32 SectionVertexIndex = INDEX_NONE;
+		LODData.GetSectionFromVertexIndex(VertexIndex, SectionIndex, SectionVertexIndex);
+
+		FQuat4f QuatSum = FQuat4f(0, 0, 0, 0);
+
+		FQuat4f R0;
+		float Sign = 1;
+		const int32 NumInfluences = SkinWeightBuffer.GetMaxBoneInfluences();
+		for (int32 InfluenceIndex = 0; InfluenceIndex < NumInfluences; ++InfluenceIndex)
+		{
+			const int32 BoneIndex = SkinWeightBuffer.GetBoneIndex(VertexIndex, InfluenceIndex);
+			const uint16 WeightByte = SkinWeightBuffer.GetBoneWeight(VertexIndex, InfluenceIndex);
+			// Weight must be > 0 when InflueceIndex == 0
+			ensure(InfluenceIndex > 0 || WeightByte > 0);
+			if (WeightByte > 0)
+			{
+				const int32 RealBoneIndex = LODData.RenderSections[SectionIndex].BoneMap[BoneIndex];
+				const FQuat4f R(BoneMatrices[RealBoneIndex].GetMatrixWithoutScale());
+				if (InfluenceIndex == 0)
+				{
+					R0 = R; Sign = 1;
+				}
+				else
+				{
+					Sign = Inner(R0, R) < 0 ? -1 : 1;
+				}
+				const float Weight = Sign * static_cast<float>(WeightByte) * UE::AnimationCore::InvMaxRawBoneWeightFloat;
+				QuatSum += R * Weight;
+			}
+		}
+
+		const float SizeSquared = QuatSum.SizeSquared();
+		if (SizeSquared > SMALL_NUMBER)
+		{
+			/** Unrotate vector using v' = q^{-1} v q if q is unit size 
+			 * Because QuatSum is not unit size
+			 * v' =  q^* v q / |q|^2
+			 */
+			return ToVector<FQuat4f>(Conjugate(QuatSum) * FromVector<FQuat4f>(WorldDelta) * QuatSum) / SizeSquared;
+		}
+		else
+		{
+			return WorldDelta;			
+		}
+	}
+
 	void FMLDeformerSampler::ExtractUnskinnedPositions(int32 LODIndex, TArray<FVector3f>& OutPositions) const
 	{
 		OutPositions.Reset();
diff --git a/Engine/Plugins/Animation/MLDeformer/MLDeformerFramework/Source/MLDeformerFrameworkEditor/Private/SMLDeformerTimeline.cpp b/Engine/Plugins/Animation/MLDeformer/MLDeformerFramework/Source/MLDeformerFrameworkEditor/Private/SMLDeformerTimeline.cpp
--- a/Engine/Plugins/Animation/MLDeformer/MLDeformerFramework/Source/MLDeformerFrameworkEditor/Private/SMLDeformerTimeline.cpp
+++ b/Engine/Plugins/Animation/MLDeformer/MLDeformerFramework/Source/MLDeformerFrameworkEditor/Private/SMLDeformerTimeline.cpp
@@ -1529,7 +1529,7 @@ namespace UE::MLDeformer
 										TSharedPtr<FMLDeformerEditorModel> EditorModel = Model.Pin();
 										check(EditorModel.IsValid());										
 										const EMLDeformerVizMode VizMode = EditorModel->GetModel()->GetVizSettings()->GetVisualizationMode();
-										return (VizMode == EMLDeformerVizMode::TrainingData && EditorModel->GetNumTrainingInputAnims() > 0) ? EVisibility::Visible : EVisibility::Hidden;
+										return (VizMode == EMLDeformerVizMode::TrainingData) ? EVisibility::Visible : EVisibility::Hidden;
 									}
 								)
 								[
@@ -1908,9 +1908,28 @@ namespace UE::MLDeformer
 
 	void SMLDeformerTimeline::SetTrainingAnimNames(const TArray<TSharedPtr<FMLDeformerTrainingInputAnimName>>& Names)
 	{
+		const int32 OldSelectedItemIndex = SelectedTrainingAnimNameIndex;
+
+		// Detect when the selected item index needs to be updated as the position in the name list changed.
+		TSharedPtr<FMLDeformerTrainingInputAnimName> PrevSelected;
+		if (SelectedTrainingAnimNameIndex != INDEX_NONE)
+		{
+			PrevSelected = TrainingAnimNames[SelectedTrainingAnimNameIndex];
+			if (Names.IsValidIndex(SelectedTrainingAnimNameIndex) && PrevSelected != Names[SelectedTrainingAnimNameIndex])
+			{
+				for (int32 Index = 0; Index < Names.Num(); ++Index)
+				{
+					if (Names[Index].IsValid() && Names[Index]->Name == PrevSelected->Name)
+					{
+						SelectedTrainingAnimNameIndex = Index;
+						break;
+					}
+				}
+			}
+		}
+
 		TrainingAnimNames = Names;
 
-		const int32 OldSelectedItemIndex = SelectedTrainingAnimNameIndex;
 		if (!Names.IsEmpty() && SelectedTrainingAnimNameIndex == INDEX_NONE)
 		{
 			SelectedTrainingAnimNameIndex = 0;
diff --git a/Engine/Plugins/Animation/MLDeformer/MLDeformerFramework/Source/MLDeformerFrameworkEditor/Public/MLDeformerEditorModel.h b/Engine/Plugins/Animation/MLDeformer/MLDeformerFramework/Source/MLDeformerFrameworkEditor/Public/MLDeformerEditorModel.h
--- a/Engine/Plugins/Animation/MLDeformer/MLDeformerFramework/Source/MLDeformerFrameworkEditor/Public/MLDeformerEditorModel.h
+++ b/Engine/Plugins/Animation/MLDeformer/MLDeformerFramework/Source/MLDeformerFrameworkEditor/Public/MLDeformerEditorModel.h
@@ -578,6 +578,15 @@ namespace UE::MLDeformer
 		 */
 		virtual void CreateHeatMapAssets();
 
+		/**
+		 * Get the heat map deformer that should currently be used (when heatmaps are enabled).
+		 * On default this will return the HeatMapDeformerGraph.
+		 * However, if your deformer supports let's say linear blend skinning and dual quaternion skinning, you could
+		 * make it return the appropriate deformer graph for skinning method that is currently used.
+		 * @return Returns a pointer to the mesh deformer to use when heatmaps are enabled.
+		 */
+		virtual UMeshDeformer* GetActiveHeatMapDeformer() const;
+
 		/**
 		 * Specify whether the heat map material should be enabled or not.
 		 * This will activate the heat map visualization, or deactivate it.
@@ -654,11 +663,17 @@ namespace UE::MLDeformer
 		virtual FString GetHeatMapMaterialPath() const;
 
 		/**
-		 * Get the path to the heat map deformer graph asset.
+		 * Get the path to the heat map deformer graph asset, when using linear skinning.
 		 * @return The path to the heat map deformer graph asset.
 		 */
 		virtual FString GetHeatMapDeformerGraphPath() const;
 
+		/**
+		 * Get the path to the heat map deformer graph asset, when using dual quaternion skinning.
+		 * @return The path to the heat map deformer graph asset.
+		 */
+		virtual FString GetHeatMapDeformerGraphDualQuatPath() const;
+
 		/**
 		 * Refresh the contents of the input widget.
 		 * This includes the bone and curve list contents.
@@ -1222,6 +1237,9 @@ namespace UE::MLDeformer
 		/** The heatmap deformer graph. */
 		TObjectPtr<UMeshDeformer> HeatMapDeformerGraph = nullptr;
 
+		/** The dual quaternion based heat map. */
+		TObjectPtr<UMeshDeformer> HeatMapDeformerGraphDualQuat = nullptr;
+
 		/** A list of vertex attribute names that can be used inside combobox UI elements. */
 		TArray<TSharedPtr<FName>> VertexAttributeNames;
 
diff --git a/Engine/Plugins/Animation/MLDeformer/MLDeformerFramework/Source/MLDeformerFrameworkEditor/Public/MLDeformerMorphModelEditorModel.h b/Engine/Plugins/Animation/MLDeformer/MLDeformerFramework/Source/MLDeformerFrameworkEditor/Public/MLDeformerMorphModelEditorModel.h
--- a/Engine/Plugins/Animation/MLDeformer/MLDeformerFramework/Source/MLDeformerFrameworkEditor/Public/MLDeformerMorphModelEditorModel.h
+++ b/Engine/Plugins/Animation/MLDeformer/MLDeformerFramework/Source/MLDeformerFrameworkEditor/Public/MLDeformerMorphModelEditorModel.h
@@ -33,6 +33,7 @@ namespace UE::MLDeformer
 		virtual void Render(const FSceneView* View, FViewport* Viewport, FPrimitiveDrawInterface* PDI) override;
 		virtual void OnPropertyChanged(FPropertyChangedEvent& PropertyChangedEvent) override;
 		virtual FString GetHeatMapDeformerGraphPath() const override;
+		virtual FString GetHeatMapDeformerGraphDualQuatPath() const override;
 		virtual void OnPreTraining() override;
 		virtual void OnPostTraining(ETrainingResult TrainingResult, bool bUsePartiallyTrainedWhenAborted) override;
 		virtual void OnMaxNumLODsChanged() override;
diff --git a/Engine/Plugins/Animation/MLDeformer/MLDeformerFramework/Source/MLDeformerFrameworkEditor/Public/MLDeformerSampler.h b/Engine/Plugins/Animation/MLDeformer/MLDeformerFramework/Source/MLDeformerFrameworkEditor/Public/MLDeformerSampler.h
--- a/Engine/Plugins/Animation/MLDeformer/MLDeformerFramework/Source/MLDeformerFrameworkEditor/Public/MLDeformerSampler.h
+++ b/Engine/Plugins/Animation/MLDeformer/MLDeformerFramework/Source/MLDeformerFrameworkEditor/Public/MLDeformerSampler.h
@@ -20,10 +20,13 @@ namespace UE::MLDeformer
 	/**
 	 * The space in which the sampler will calculate its vertex deltas.
 	 */
-	enum class EVertexDeltaSpace
+	enum class EVertexDeltaSpace : int8
 	{
-		PreSkinning,	/** Vertex deltas will be in pre-skinning space. We train on pre-skinning deltas. */
-		PostSkinning	/** Vertex deltas will be in post-skinning space. Mostly used to visualize deltas in the viewport. */
+		/** Vertex deltas will be in pre-skinning space. We train on pre-skinning deltas. */
+		PreSkinning,
+
+		/** Vertex deltas will be in post-skinning space. Mostly used to visualize deltas in the viewport. */
+		PostSkinning
 	};
 
 	/**
@@ -61,7 +64,7 @@ namespace UE::MLDeformer
 		 * Get the current bone rotations that we can feed to a neural network. 
 		 * The number of floats in the buffer equals to NumBones * 6.
 		 * The six floats represent two columns of the 3x3 rotation matrix of the bone.
-		 * @return The arrray of floats that we can pass to the neural network as inputs.
+		 * @return The array of floats that we can pass to the neural network as inputs.
 		 */
 		const TArray<float>& GetBoneRotations() const					{ return BoneRotations; }
 
@@ -77,6 +80,12 @@ namespace UE::MLDeformer
 		/** Set the space in which we should calculate the vertex deltas. */
 		void SetVertexDeltaSpace(EVertexDeltaSpace DeltaSpace)			{ VertexDeltaSpace = DeltaSpace; }
 
+		/** Get the skinning mode. */
+		EMLDeformerSkinningMode GetSkinningMode() const					{ return SkinningMode; }
+
+		/** Set the skinning mode. This defines the type of skinning we use, like Linear Blend Skinning or Dual Quaternion skinning. */
+		void SetSkinningMode(const EMLDeformerSkinningMode Mode)		{ SkinningMode = Mode; }
+
 		/** Have we already initialized this sampler? Is it ready for sampling? */
 		bool IsInitialized() const										{ return SkelMeshActor.Get() != nullptr; }
 
@@ -116,6 +125,9 @@ namespace UE::MLDeformer
 		/** Calculate the inverse skinning transform, which is used to convert from post-skinning to pre-skinning space. */
 		FMatrix44f CalcInverseSkinningTransform(int32 VertexIndex, const FSkeletalMeshLODRenderData& SkelMeshLODData, const FSkinWeightVertexBuffer& SkinWeightBuffer) const;
 
+		/** Calculate the delta using dual quaternions, taking a world space delta vector as input. */
+		FVector3f CalcDualQuaternionDelta(int32 VertexIndex, const FVector3f& WorldDelta, const FSkeletalMeshLODRenderData& SkelMeshLODData, const FSkinWeightVertexBuffer& SkinWeightBuffer) const;
+
 		/** Create a new actor with a specific name, in a specified world. */
 		AActor* CreateNewActor(UWorld* InWorld, const FName& Name) const;
 
@@ -196,5 +208,8 @@ namespace UE::MLDeformer
 
 		/** The vertex delta space (pre or post skinning) used when calculating the deltas. */
 		EVertexDeltaSpace VertexDeltaSpace = EVertexDeltaSpace::PreSkinning;
+
+		/** The skinning mode, which only works in combination with pre-skinning space. */
+		EMLDeformerSkinningMode SkinningMode = EMLDeformerSkinningMode::Linear;
 	};
 }	// namespace UE::MLDeformer
diff --git a/Engine/Plugins/Animation/MLDeformer/NearestNeighborModel/NearestNeighborModel.uplugin b/Engine/Plugins/Animation/MLDeformer/NearestNeighborModel/NearestNeighborModel.uplugin
--- a/Engine/Plugins/Animation/MLDeformer/NearestNeighborModel/NearestNeighborModel.uplugin
+++ b/Engine/Plugins/Animation/MLDeformer/NearestNeighborModel/NearestNeighborModel.uplugin
@@ -2,8 +2,8 @@
 	"FileVersion": 3,
 	"Version": 1,
 	"VersionName": "0.3",
-	"FriendlyName": "ML Deformer Nearest Neighbor Model",
-	"Description": "Nearest Neighbor Model for the ML Deformer Framework",
+	"FriendlyName": "ML Deformer Nearest Neighbor Model (DEPRECATED)",
+	"Description": "Nearest Neighbor Model for the ML Deformer Framework. This model has been deprecated. Please use the Detail Pose Model instead.",
 	"Category": "Animation",
 	"CreatedBy": "Epic Games, Inc.",
 	"CreatedByURL": "https://epicgames.com",
diff --git a/Engine/Plugins/Animation/MLDeformer/NearestNeighborModel/Source/NearestNeighborModelEditor/Private/NearestNeighborModelDetails.cpp b/Engine/Plugins/Animation/MLDeformer/NearestNeighborModel/Source/NearestNeighborModelEditor/Private/NearestNeighborModelDetails.cpp
--- a/Engine/Plugins/Animation/MLDeformer/NearestNeighborModel/Source/NearestNeighborModelEditor/Private/NearestNeighborModelDetails.cpp
+++ b/Engine/Plugins/Animation/MLDeformer/NearestNeighborModel/Source/NearestNeighborModelEditor/Private/NearestNeighborModelDetails.cpp
@@ -596,6 +596,19 @@ namespace UE::NearestNeighborModel
 		// Create all the detail categories and add the properties of the base class.
 		FMLDeformerMorphModelDetails::CustomizeDetails(DetailBuilder);
 
+		// Show a warning about this model being deprecated.
+		FDetailWidgetRow& GeomCacheMappingErrorRow = BaseMeshCategoryBuilder->AddCustomRow(FText::FromString("ModelDeprecationWarning"))
+			.WholeRowContent()
+			[
+				SNew(SBox)
+				.Padding(FMargin(0.0f, 4.0f))
+				[
+					SNew(SWarningOrErrorBox)
+					.MessageStyle(EMessageStyle::Warning)
+					.Message(LOCTEXT("ModelDeprecationWarningText", "The Nearest Neighbor Model has been deprecated and is being replaced by the Detail Pose Model. Please use that model instead."))
+				]
+			];
+
 		CustomizeTrainingSettingsCategory();
 		CustomizeNearestNeighborSettingsCategory();
 		CustomizeSectionsCategory(DetailBuilder);
diff --git a/Engine/Plugins/Animation/MLDeformer/NeuralMorphModel/Source/NeuralMorphModel/Private/NeuralMorphInputInfo.cpp b/Engine/Plugins/Animation/MLDeformer/NeuralMorphModel/Source/NeuralMorphModel/Private/NeuralMorphInputInfo.cpp
--- a/Engine/Plugins/Animation/MLDeformer/NeuralMorphModel/Source/NeuralMorphModel/Private/NeuralMorphInputInfo.cpp
+++ b/Engine/Plugins/Animation/MLDeformer/NeuralMorphModel/Source/NeuralMorphModel/Private/NeuralMorphInputInfo.cpp
@@ -191,3 +191,13 @@ bool UNeuralMorphInputInfo::HasInvalidGroups() const
 
 	return false;
 }
+
+EMLDeformerSkinningMode UNeuralMorphInputInfo::GetSkinningMode() const
+{
+	return SkinningMode;
+}
+
+void UNeuralMorphInputInfo::SetSkinningMode(EMLDeformerSkinningMode Mode)
+{
+	SkinningMode = Mode;
+}
diff --git a/Engine/Plugins/Animation/MLDeformer/NeuralMorphModel/Source/NeuralMorphModel/Private/NeuralMorphModel.cpp b/Engine/Plugins/Animation/MLDeformer/NeuralMorphModel/Source/NeuralMorphModel/Private/NeuralMorphModel.cpp
--- a/Engine/Plugins/Animation/MLDeformer/NeuralMorphModel/Source/NeuralMorphModel/Private/NeuralMorphModel.cpp
+++ b/Engine/Plugins/Animation/MLDeformer/NeuralMorphModel/Source/NeuralMorphModel/Private/NeuralMorphModel.cpp
@@ -270,6 +270,7 @@ void UNeuralMorphModel::PostLoad()
 {
 	Super::PostLoad();
 	UpdateMissingGroupNames();
+	check(SupportsGlobalModeOnly() ? Mode == ENeuralMorphMode::Global : true);
 }
 
 void UNeuralMorphModel::SetNeuralMorphNetwork(UNeuralMorphNetwork* Net)
@@ -323,6 +324,7 @@ void UNeuralMorphModel::GetAssetRegistryTags(FAssetRegistryTagsContext Context)
 	Context.AddTag(FAssetRegistryTag("MLDeformer.NeuralMorphModel.RegularizationFactor", FString::Printf(TEXT("%f"), RegularizationFactor), FAssetRegistryTag::TT_Numerical));
 	Context.AddTag(FAssetRegistryTag("MLDeformer.NeuralMorphModel.SmoothLossBeta", FString::Printf(TEXT("%f"), SmoothLossBeta), FAssetRegistryTag::TT_Numerical));
 	Context.AddTag(FAssetRegistryTag("MLDeformer.NeuralMorphModel.EnableBoneMasks", bEnableBoneMasks ? TEXT("True") : TEXT("False"), FAssetRegistryTag::TT_Alphabetical));
+	Context.AddTag(FAssetRegistryTag("MLDeformer.NeuralMorphModel.UseDualQuaternions", SkinningMode == EMLDeformerSkinningMode::DualQuaternion ? TEXT("True") : TEXT("False"), FAssetRegistryTag::TT_Alphabetical));
 }
 
 #undef LOCTEXT_NAMESPACE
diff --git a/Engine/Plugins/Animation/MLDeformer/NeuralMorphModel/Source/NeuralMorphModel/Private/NeuralMorphModelInstance.cpp b/Engine/Plugins/Animation/MLDeformer/NeuralMorphModel/Source/NeuralMorphModel/Private/NeuralMorphModelInstance.cpp
--- a/Engine/Plugins/Animation/MLDeformer/NeuralMorphModel/Source/NeuralMorphModel/Private/NeuralMorphModelInstance.cpp
+++ b/Engine/Plugins/Animation/MLDeformer/NeuralMorphModel/Source/NeuralMorphModel/Private/NeuralMorphModelInstance.cpp
@@ -237,7 +237,7 @@ void UNeuralMorphModelInstance::Execute(float ModelWeight)
 	}
 
 	// Make sure we have a valid network instance.
-	if (NetworkInstance == nullptr)
+	if (!NetworkInstance || MorphNetwork->IsEmpty())
 	{
 		WeightData->ZeroWeights();
 		return;
@@ -248,10 +248,10 @@ void UNeuralMorphModelInstance::Execute(float ModelWeight)
 	// Get the network output values, read the values and use them as morph target weights inside the skeletal mesh component.
 	const TArrayView<const float> NetworkOutputs = NetworkInstance->GetOutputs();
 	const int32 NumNetworkWeights = NetworkOutputs.Num();
-	const int32 NumMorphTargets = WeightData->Weights.Num();
+	const int32 NumMorphWeights = WeightData->Weights.Num();
 
-	// Set the weights to zero if we have no valid network.
-	if (MorphNetwork->IsEmpty() || NumMorphTargets != NumNetworkWeights + 1)	// Plus 1 because of the extra means morph target.
+	// Set the weights to zero if we don't have enough weights.
+	if (NumMorphWeights < NumNetworkWeights + 1)	// Plus 1 because of the extra means morph target.
 	{
 		WeightData->ZeroWeights();
 		return;
diff --git a/Engine/Plugins/Animation/MLDeformer/NeuralMorphModel/Source/NeuralMorphModel/Public/NeuralMorphInputInfo.h b/Engine/Plugins/Animation/MLDeformer/NeuralMorphModel/Source/NeuralMorphModel/Public/NeuralMorphInputInfo.h
--- a/Engine/Plugins/Animation/MLDeformer/NeuralMorphModel/Source/NeuralMorphModel/Public/NeuralMorphInputInfo.h
+++ b/Engine/Plugins/Animation/MLDeformer/NeuralMorphModel/Source/NeuralMorphModel/Public/NeuralMorphInputInfo.h
@@ -4,6 +4,7 @@
 #include "CoreTypes.h"
 #include "MLDeformerMorphModelInputInfo.h"
 #include "NeuralMorphTypes.h"
+#include "MLDeformerModel.h"
 #include "NeuralMorphInputInfo.generated.h"
 
 /**
@@ -74,6 +75,17 @@ public:
 	 */
 	bool HasInvalidGroups() const;
 
+	/**
+	 * Get the skinning mode that we trained with.
+	 * @return Returns the skinning mode, such as linear skinned or dual quaternion skinned.
+	 */
+	EMLDeformerSkinningMode GetSkinningMode() const;
+
+	/**
+	 * Set the skinning mode that was used to train with.
+	 */
+	void SetSkinningMode(EMLDeformerSkinningMode Mode);
+
 protected:
 	/** The groups of bones that generate morph targets together. */
 	UPROPERTY()
@@ -82,4 +94,8 @@ protected:
 	/** The groups of curves that generate morph targets together. */
 	UPROPERTY()
 	TArray<FNeuralMorphCurveGroup> CurveGroups;
+
+	/** The skinning mode that was trained for. */
+	UPROPERTY()
+	EMLDeformerSkinningMode SkinningMode = EMLDeformerSkinningMode::Linear;
 };
diff --git a/Engine/Plugins/Animation/MLDeformer/NeuralMorphModel/Source/NeuralMorphModel/Public/NeuralMorphModel.h b/Engine/Plugins/Animation/MLDeformer/NeuralMorphModel/Source/NeuralMorphModel/Public/NeuralMorphModel.h
--- a/Engine/Plugins/Animation/MLDeformer/NeuralMorphModel/Source/NeuralMorphModel/Public/NeuralMorphModel.h
+++ b/Engine/Plugins/Animation/MLDeformer/NeuralMorphModel/Source/NeuralMorphModel/Public/NeuralMorphModel.h
@@ -55,10 +55,17 @@ public:
 	virtual bool IsTrained() const override;
 	// ~END UMLDeformerModel overrides.
 
+	/**
+	 * Check if only global mode should be supported or not.
+	 * Models that inherit from this model can disable this if they only want the global mode.
+	 * This will automatically hide the UI elements related to the mode.
+	 */
+	virtual bool SupportsGlobalModeOnly() const						{ return false; }
+
 	const TArray<FNeuralMorphBoneGroup>& GetBoneGroups() const		{ return BoneGroups; }
 	const TArray<FNeuralMorphCurveGroup>& GetCurveGroups() const	{ return CurveGroups; }
 	UNeuralMorphNetwork* GetNeuralMorphNetwork() const				{ return NeuralMorphNetwork.Get(); }
-	ENeuralMorphMode GetModelMode() const							{ return Mode; }
+	ENeuralMorphMode GetModelMode() const							{ return SupportsGlobalModeOnly() ? ENeuralMorphMode::Global : Mode; }
 	int32 GetLocalNumMorphsPerBone() const							{ return LocalNumMorphTargetsPerBone; }
 	int32 GetLocalNumHiddenLayers() const							{ return LocalNumHiddenLayers; }
 	int32 GetLocalNumNeuronsPerLayer() const						{ return LocalNumNeuronsPerLayer; }
@@ -73,10 +80,16 @@ public:
 	float GetRegularizationFactor() const							{ return RegularizationFactor; }
 	bool IsBoneMaskingEnabled() const								{ return bEnableBoneMasks; }
 
+	EMLDeformerSkinningMode GetSkinningMode() const					{ return SkinningMode; }
+	void SetSkinningMode(EMLDeformerSkinningMode SkinMode)			{ SkinningMode = SkinMode; }
+
 	void UpdateMissingGroupNames();
 	void SetNeuralMorphNetwork(UNeuralMorphNetwork* Net);
 	void SetNumIterations(int32 InNumIterations)					{ check(InNumIterations > 0); NumIterations = InNumIterations; }
 
+	static FName GetModePropertyName()								{ return GET_MEMBER_NAME_CHECKED(UNeuralMorphModel, Mode); }
+	static FName GetSkinningModePropertyName()						{ return GET_MEMBER_NAME_CHECKED(UNeuralMorphModel, SkinningMode); }
+
 public:
 	/**
 	 * The set of bones that are grouped together and generate morph targets together as a whole.
@@ -198,7 +211,7 @@ public:
 	 * The benefit of enabling this can be reduced GPU memory footprint, faster GPU performance and more localized deformations.
 	 * If deformations do not happen near the joint then this enabling this setting can lead to those deformations possibly not being captured.
 	 */
-	UPROPERTY(EditAnywhere, BlueprintReadWrite, AdvancedDisplay, Category = "Training Settings", meta = (EditCondition = "Mode == ENeuralMorphMode::Local"))
+	UPROPERTY(EditAnywhere, BlueprintReadWrite, AdvancedDisplay, Category = "Training Settings", meta = (EditCondition = "Mode == ENeuralMorphMode::Local", EditConditionHides))
 	bool bEnableBoneMasks = false;
 
 	/**
@@ -210,6 +223,14 @@ public:
 	UPROPERTY(EditAnywhere, BlueprintReadWrite, AdvancedDisplay, Category = "Training Settings", meta = (ClampMin = "0.0", ClampMax = "10.0"))
 	float SmoothLossBeta = 0.0f;
 
+	/**
+	 * The base skinning mode to use. The learned deltas will be applied on top of this basis.
+	 * Please keep in mind that you use the same type of skinning on your character, otherwise the applied deltas are wrong.
+	 * If you train using Dual Quaternions you should also run a Dual Quaternion skinning deformer before applying the ML Deltas on top.
+	 */
+	UPROPERTY(EditAnywhere, AdvancedDisplay, Category = "Training Settings")
+	EMLDeformerSkinningMode SkinningMode = EMLDeformerSkinningMode::Linear;
+
 	/**
 	 * The neural morph model network.
 	 * This is a specialized network used for inference, for performance reasons.
diff --git a/Engine/Plugins/Animation/MLDeformer/NeuralMorphModel/Source/NeuralMorphModel/Public/NeuralMorphModelInstance.h b/Engine/Plugins/Animation/MLDeformer/NeuralMorphModel/Source/NeuralMorphModel/Public/NeuralMorphModelInstance.h
--- a/Engine/Plugins/Animation/MLDeformer/NeuralMorphModel/Source/NeuralMorphModel/Public/NeuralMorphModelInstance.h
+++ b/Engine/Plugins/Animation/MLDeformer/NeuralMorphModel/Source/NeuralMorphModel/Public/NeuralMorphModelInstance.h
@@ -22,6 +22,8 @@ public:
 	virtual void Execute(float ModelWeight) override;
 	// ~END UMLDeformerModelInstance overrides.
 
+	UNeuralMorphNetworkInstance* GetNetworkInstance() const	{ return NetworkInstance.Get(); }
+
 protected:
 	/**
 	 * Set the network inputs.
diff --git a/Engine/Plugins/Animation/MLDeformer/NeuralMorphModel/Source/NeuralMorphModel/Public/NeuralMorphModelVizSettings.h b/Engine/Plugins/Animation/MLDeformer/NeuralMorphModel/Source/NeuralMorphModel/Public/NeuralMorphModelVizSettings.h
--- a/Engine/Plugins/Animation/MLDeformer/NeuralMorphModel/Source/NeuralMorphModel/Public/NeuralMorphModelVizSettings.h
+++ b/Engine/Plugins/Animation/MLDeformer/NeuralMorphModel/Source/NeuralMorphModel/Public/NeuralMorphModelVizSettings.h
@@ -5,12 +5,13 @@
 #include "CoreTypes.h"
 #include "MLDeformerMorphModelVizSettings.h"
 #include "NeuralMorphTypes.h"
+#include "NeuralMorphModel.h"
 #include "NeuralMorphModelVizSettings.generated.h"
 
 /**
- * The vizualization settings specific to this model.
+ * The visualization settings specific to this model.
  * Even if we have no new properties compared to the morph model, we still need to
- * create this class in order to propertly register a detail customization for it in our editor module.
+ * create this class in order to properly register a detail customization for it in our editor module.
  */
 UCLASS()
 class NEURALMORPHMODEL_API UNeuralMorphModelVizSettings
@@ -19,6 +20,12 @@ class NEURALMORPHMODEL_API UNeuralMorphModelVizSettings
 	GENERATED_BODY()
 
 public:
+	UFUNCTION()
+	bool ShouldShowMaskVizMode() const
+	{
+		return Cast<UNeuralMorphModel>(GetOuter())->GetModelMode() == ENeuralMorphMode::Local;
+	}
+
 	/**
 	 * The visualization mode for the masks.
 	 * Each bone, curve, bone group or curve group has a specific mask area on the mesh.
@@ -26,6 +33,6 @@ public:
 	 * For example if you rotate the left arm, you don't want the right arm to deform. The mask for the left arm can be setup in a way that it only includes 
 	 * vertices around the area of the left arm to enforce this.
 	 */
-	UPROPERTY(EditAnywhere, Category = "Live Settings")
+	UPROPERTY(EditAnywhere, Category = "Live Settings", meta = (EditCondition = "ShouldShowMaskVizMode()", EditConditionHides))
 	ENeuralMorphMaskVizMode MaskVizMode = ENeuralMorphMaskVizMode::WhenInFocus;
 };
diff --git a/Engine/Plugins/Animation/MLDeformer/NeuralMorphModel/Source/NeuralMorphModelEditor/Private/NeuralMorphEditorModel.cpp b/Engine/Plugins/Animation/MLDeformer/NeuralMorphModel/Source/NeuralMorphModelEditor/Private/NeuralMorphEditorModel.cpp
--- a/Engine/Plugins/Animation/MLDeformer/NeuralMorphModel/Source/NeuralMorphModelEditor/Private/NeuralMorphEditorModel.cpp
+++ b/Engine/Plugins/Animation/MLDeformer/NeuralMorphModel/Source/NeuralMorphModelEditor/Private/NeuralMorphEditorModel.cpp
@@ -68,9 +68,9 @@ namespace UE::NeuralMorphModel
 			UpdateIsReadyForTrainingState();
 			RebuildEditorMaskInfo();
 		}
-		else if (Property->GetFName() == GET_MEMBER_NAME_CHECKED(UNeuralMorphModel, Mode))
+		else if (Property->GetFName() == UNeuralMorphModel::GetModePropertyName())
 		{
-			if (PropertyChangedEvent.ChangeType == EPropertyChangeType::ValueSet)
+			if (PropertyChangedEvent.ChangeType & EPropertyChangeType::ValueSet || PropertyChangedEvent.ChangeType & EPropertyChangeType::ResetToDefault)
 			{
 				SetResamplingInputOutputsNeeded(true);
 				UpdateIsReadyForTrainingState();
@@ -78,6 +78,17 @@ namespace UE::NeuralMorphModel
 				RebuildEditorMaskInfo();
 			}
 		}
+		else if (Property->GetFName() == UNeuralMorphModel::GetSkinningModePropertyName())
+		{	
+			SetResamplingInputOutputsNeeded(true);
+			for (TSharedPtr<FMLDeformerSampler> Sampler : Samplers)
+			{
+				if (Sampler.IsValid())
+				{
+					Sampler->SetSkinningMode(GetNeuralMorphModel()->GetSkinningMode());
+				}
+			}
+		}
 	}
 
 	TSharedPtr<SMLDeformerInputWidget> FNeuralMorphEditorModel::CreateInputWidget()
@@ -99,12 +110,34 @@ namespace UE::NeuralMorphModel
 		}
 	}
 
+	UMeshDeformer* FNeuralMorphEditorModel::GetActiveHeatMapDeformer() const
+	{
+		if (Model->IsTrained())
+		{
+			UNeuralMorphInputInfo* NeuralMorphInputInfo = Cast<UNeuralMorphInputInfo>(Model->GetInputInfo());
+			check(NeuralMorphInputInfo);
+
+			const EMLDeformerSkinningMode SkinningMode = NeuralMorphInputInfo->GetSkinningMode();
+			if (SkinningMode == EMLDeformerSkinningMode::DualQuaternion)
+			{
+				return HeatMapDeformerGraphDualQuat;
+			}
+		}
+
+		// Default linear skinned heatmap.
+		return HeatMapDeformerGraph;
+	}
+
 	void FNeuralMorphEditorModel::Init(const InitSettings& Settings)
 	{
 		FMLDeformerMorphModelEditorModel::Init(Settings);
 
+		UNeuralMorphModel* NeuralModel = GetNeuralMorphModel();
+
+		check(NeuralModel->SupportsGlobalModeOnly() ? NeuralModel->Mode == ENeuralMorphMode::Global : true);
+
 		RemoveNonExistingMaskInfos();
-		if (GetNeuralMorphModel()->GetModelMode() == ENeuralMorphMode::Local)
+		if (NeuralModel->GetModelMode() == ENeuralMorphMode::Local)
 		{
 			RebuildEditorMaskInfo();
 		}
@@ -156,6 +189,8 @@ namespace UE::NeuralMorphModel
 		}
 
 		UNeuralMorphModel* NeuralMorphModel = Cast<UNeuralMorphModel>(Model);
+		NeuralInputInfo->SetSkinningMode(NeuralMorphModel->GetSkinningMode());
+
 		if (NeuralMorphModel->GetModelMode() != ENeuralMorphMode::Local)
 		{
 			return;
diff --git a/Engine/Plugins/Animation/MLDeformer/NeuralMorphModel/Source/NeuralMorphModelEditor/Private/NeuralMorphModelDetails.cpp b/Engine/Plugins/Animation/MLDeformer/NeuralMorphModel/Source/NeuralMorphModelEditor/Private/NeuralMorphModelDetails.cpp
--- a/Engine/Plugins/Animation/MLDeformer/NeuralMorphModel/Source/NeuralMorphModelEditor/Private/NeuralMorphModelDetails.cpp
+++ b/Engine/Plugins/Animation/MLDeformer/NeuralMorphModel/Source/NeuralMorphModelEditor/Private/NeuralMorphModelDetails.cpp
@@ -30,23 +30,24 @@ namespace UE::NeuralMorphModel
 		UNeuralMorphModel* NeuralMorphModel = Cast<UNeuralMorphModel>(Model);
 		check(NeuralMorphModel);
 
-		TrainingSettingsCategoryBuilder->AddProperty(GET_MEMBER_NAME_CHECKED(UNeuralMorphModel, Mode), UNeuralMorphModel::StaticClass());
+		TrainingSettingsCategoryBuilder->AddProperty(GET_MEMBER_NAME_CHECKED(UNeuralMorphModel, Mode), UNeuralMorphModel::StaticClass())
+			.Visibility(NeuralMorphModel->SupportsGlobalModeOnly() ? EVisibility::Collapsed : EVisibility::Visible);
 
 		// Local mode settings.
 		TrainingSettingsCategoryBuilder->AddProperty(GET_MEMBER_NAME_CHECKED(UNeuralMorphModel, LocalNumMorphTargetsPerBone), UNeuralMorphModel::StaticClass())
-			.Visibility(NeuralMorphModel->Mode == ENeuralMorphMode::Local ? EVisibility::Visible : EVisibility::Collapsed);
+			.Visibility(NeuralMorphModel->GetModelMode() == ENeuralMorphMode::Local ? EVisibility::Visible : EVisibility::Collapsed);
 		TrainingSettingsCategoryBuilder->AddProperty(GET_MEMBER_NAME_CHECKED(UNeuralMorphModel, LocalNumHiddenLayers), UNeuralMorphModel::StaticClass())
-			.Visibility(NeuralMorphModel->Mode == ENeuralMorphMode::Local ? EVisibility::Visible : EVisibility::Collapsed);
+			.Visibility(NeuralMorphModel->GetModelMode() == ENeuralMorphMode::Local ? EVisibility::Visible : EVisibility::Collapsed);
 		TrainingSettingsCategoryBuilder->AddProperty(GET_MEMBER_NAME_CHECKED(UNeuralMorphModel, LocalNumNeuronsPerLayer), UNeuralMorphModel::StaticClass())
-			.Visibility(NeuralMorphModel->Mode == ENeuralMorphMode::Local ? EVisibility::Visible : EVisibility::Collapsed);
+			.Visibility(NeuralMorphModel->GetModelMode() == ENeuralMorphMode::Local ? EVisibility::Visible : EVisibility::Collapsed);
 
 		// Global mode settings.
 		TrainingSettingsCategoryBuilder->AddProperty(GET_MEMBER_NAME_CHECKED(UNeuralMorphModel, GlobalNumMorphTargets), UNeuralMorphModel::StaticClass())
-			.Visibility(NeuralMorphModel->Mode == ENeuralMorphMode::Global ? EVisibility::Visible : EVisibility::Collapsed);
+			.Visibility(NeuralMorphModel->GetModelMode() == ENeuralMorphMode::Global ? EVisibility::Visible : EVisibility::Collapsed);
 		TrainingSettingsCategoryBuilder->AddProperty(GET_MEMBER_NAME_CHECKED(UNeuralMorphModel, GlobalNumHiddenLayers), UNeuralMorphModel::StaticClass())
-			.Visibility(NeuralMorphModel->Mode == ENeuralMorphMode::Global ? EVisibility::Visible : EVisibility::Collapsed);
+			.Visibility(NeuralMorphModel->GetModelMode()== ENeuralMorphMode::Global ? EVisibility::Visible : EVisibility::Collapsed);
 		TrainingSettingsCategoryBuilder->AddProperty(GET_MEMBER_NAME_CHECKED(UNeuralMorphModel, GlobalNumNeuronsPerLayer), UNeuralMorphModel::StaticClass())
-			.Visibility(NeuralMorphModel->Mode == ENeuralMorphMode::Global ? EVisibility::Visible : EVisibility::Collapsed);
+			.Visibility(NeuralMorphModel->GetModelMode() == ENeuralMorphMode::Global ? EVisibility::Visible : EVisibility::Collapsed);
 
 		TrainingSettingsCategoryBuilder->AddProperty(GET_MEMBER_NAME_CHECKED(UNeuralMorphModel, NumIterations), UNeuralMorphModel::StaticClass());
 		TrainingSettingsCategoryBuilder->AddProperty(GET_MEMBER_NAME_CHECKED(UNeuralMorphModel, BatchSize), UNeuralMorphModel::StaticClass());
@@ -54,6 +55,7 @@ namespace UE::NeuralMorphModel
 		TrainingSettingsCategoryBuilder->AddProperty(GET_MEMBER_NAME_CHECKED(UNeuralMorphModel, RegularizationFactor), UNeuralMorphModel::StaticClass());
 		TrainingSettingsCategoryBuilder->AddProperty(GET_MEMBER_NAME_CHECKED(UNeuralMorphModel, SmoothLossBeta), UNeuralMorphModel::StaticClass());
 		TrainingSettingsCategoryBuilder->AddProperty(GET_MEMBER_NAME_CHECKED(UNeuralMorphModel, bEnableBoneMasks), UNeuralMorphModel::StaticClass());
+		TrainingSettingsCategoryBuilder->AddProperty(UNeuralMorphModel::GetSkinningModePropertyName(), UNeuralMorphModel::StaticClass());
 	}
 }	// namespace UE::NeuralMorphModel
 
diff --git a/Engine/Plugins/Animation/MLDeformer/NeuralMorphModel/Source/NeuralMorphModelEditor/Private/NeuralMorphTrainingModel.cpp b/Engine/Plugins/Animation/MLDeformer/NeuralMorphModel/Source/NeuralMorphModelEditor/Private/NeuralMorphTrainingModel.cpp
--- a/Engine/Plugins/Animation/MLDeformer/NeuralMorphModel/Source/NeuralMorphModelEditor/Private/NeuralMorphTrainingModel.cpp
+++ b/Engine/Plugins/Animation/MLDeformer/NeuralMorphModel/Source/NeuralMorphModelEditor/Private/NeuralMorphTrainingModel.cpp
@@ -4,6 +4,23 @@
 #include "NeuralMorphEditorModel.h"
 #include "NeuralMorphModel.h"
 #include "NeuralMorphInputInfo.h"
+#include "MLDeformerSampler.h"
+
+void UNeuralMorphTrainingModel::Init(UE::MLDeformer::FMLDeformerEditorModel* InEditorModel)
+{
+	Super::Init(InEditorModel);
+
+	UNeuralMorphModel* NeuralMorphModel = Cast<UNeuralMorphModel>(EditorModel->GetModel());
+	check(NeuralMorphModel);
+	for (int32 AnimIndex = 0; AnimIndex < EditorModel->GetNumTrainingInputAnims(); ++AnimIndex)
+	{
+		UE::MLDeformer::FMLDeformerSampler* Sampler = EditorModel->GetSamplerForTrainingAnim(AnimIndex);
+		if (Sampler)
+		{
+			Sampler->SetSkinningMode(NeuralMorphModel->GetSkinningMode());
+		}
+	}
+}
 
 int32 UNeuralMorphTrainingModel::GetNumBoneGroups() const
 {
diff --git a/Engine/Plugins/Animation/MLDeformer/NeuralMorphModel/Source/NeuralMorphModelEditor/Public/NeuralMorphEditorModel.h b/Engine/Plugins/Animation/MLDeformer/NeuralMorphModel/Source/NeuralMorphModelEditor/Public/NeuralMorphEditorModel.h
--- a/Engine/Plugins/Animation/MLDeformer/NeuralMorphModel/Source/NeuralMorphModelEditor/Public/NeuralMorphEditorModel.h
+++ b/Engine/Plugins/Animation/MLDeformer/NeuralMorphModel/Source/NeuralMorphModelEditor/Public/NeuralMorphEditorModel.h
@@ -44,6 +44,7 @@ namespace UE::NeuralMorphModel
 		virtual void Render(const FSceneView* View, FViewport* Viewport, FPrimitiveDrawInterface* PDI) override;
 		virtual bool GetSupportsPerTrainingInputAnimVertexMask() const override	{ return true; }
 		virtual void UpdateTrainingDeviceList() override;
+		virtual UMeshDeformer* GetActiveHeatMapDeformer() const override;
 		// ~END FMLDeformerEditorModel overrides.
 
 		// FMLDeformerMorphModelEditorModel overrides.
diff --git a/Engine/Plugins/Animation/MLDeformer/NeuralMorphModel/Source/NeuralMorphModelEditor/Public/NeuralMorphTrainingModel.h b/Engine/Plugins/Animation/MLDeformer/NeuralMorphModel/Source/NeuralMorphModelEditor/Public/NeuralMorphTrainingModel.h
--- a/Engine/Plugins/Animation/MLDeformer/NeuralMorphModel/Source/NeuralMorphModelEditor/Public/NeuralMorphTrainingModel.h
+++ b/Engine/Plugins/Animation/MLDeformer/NeuralMorphModel/Source/NeuralMorphModelEditor/Public/NeuralMorphTrainingModel.h
@@ -18,6 +18,8 @@ class NEURALMORPHMODELEDITOR_API UNeuralMorphTrainingModel
 	GENERATED_BODY()
 
 public:
+	virtual void Init(UE::MLDeformer::FMLDeformerEditorModel* InEditorModel) override;
+
 	/**
 	 * Main training function, with implementation in python.
 	 * You need to implement this method. See the UMLDeformerTrainingModel documentation for more details.
diff --git a/Engine/Plugins/Experimental/Animation/MLDeformer/DetailPoseModel/DetailPoseModel.uplugin b/Engine/Plugins/Experimental/Animation/MLDeformer/DetailPoseModel/DetailPoseModel.uplugin
new file mode 100644
--- /dev/null
+++ b/Engine/Plugins/Experimental/Animation/MLDeformer/DetailPoseModel/DetailPoseModel.uplugin
@@ -0,0 +1,35 @@
+{
+	"FileVersion": 3,
+	"Version": 1,
+	"VersionName": "0.1",
+	"FriendlyName": "ML Deformer Detail Pose Model",
+	"Description": "Detail Pose Model for the ML Deformer Framework",
+	"Category": "Animation",
+	"CreatedBy": "Epic Games, Inc.",
+	"CreatedByURL": "https://epicgames.com",
+	"DocsURL": "https://docs.unrealengine.com/5.0/en-US/using-the-machine-learning-deformer-in-unreal-engine/",
+	"MarketplaceURL": "",
+	"SupportURL": "https://forums.unrealengine.com/",
+	"CanContainContent": true,
+	"IsBetaVersion": false,
+	"IsExperimentalVersion": true,
+	"Installed": false,
+	"Modules": [
+		{
+			"Name": "DetailPoseModel",
+			"Type": "Runtime",
+			"LoadingPhase": "Default"
+		},
+		{
+			"Name": "DetailPoseModelEditor",
+			"Type": "Editor",
+			"LoadingPhase": "Default"
+		}
+	],
+	"Plugins": [
+		{
+			"Name": "NeuralMorphModel",
+			"Enabled": true
+		}
+	]
+}
diff --git a/Engine/Plugins/Experimental/Animation/MLDeformer/DetailPoseModel/Source/DetailPoseModel/DetailPoseModel.Build.cs b/Engine/Plugins/Experimental/Animation/MLDeformer/DetailPoseModel/Source/DetailPoseModel/DetailPoseModel.Build.cs
new file mode 100644
--- /dev/null
+++ b/Engine/Plugins/Experimental/Animation/MLDeformer/DetailPoseModel/Source/DetailPoseModel/DetailPoseModel.Build.cs
@@ -0,0 +1,60 @@
+// Copyright Epic Games, Inc. All Rights Reserved.
+
+using UnrealBuildTool;
+
+public class DetailPoseModel : ModuleRules
+{
+	public DetailPoseModel(ReadOnlyTargetRules Target) : base(Target)
+	{
+		PCHUsage = ModuleRules.PCHUsageMode.UseExplicitOrSharedPCHs;
+
+		PublicIncludePaths.AddRange(
+			new string[] {
+				// ... add public include paths required here ...
+			}
+		);
+				
+		
+		PrivateIncludePaths.AddRange(
+			new string[] {
+				// ... add other private include paths required here ...
+			}
+		);
+			
+		
+		PublicDependencyModuleNames.AddRange(
+			new string[]
+			{
+				"Core"
+				// ... add other public dependencies that you statically link with here ...
+			}
+		);
+			
+		
+		PrivateDependencyModuleNames.AddRange(
+			new string[]
+			{
+				"ComputeFramework",
+				"CoreUObject",
+				"Engine",
+				"GeometryCache",
+				"OptimusCore",
+				"Projects",
+				"RenderCore",
+				"RHI",
+				"MLDeformerFramework",
+				"NNE",
+				"NNERuntimeBasicCpu",
+				"NeuralMorphModel"
+			}
+		);
+		
+		
+		DynamicallyLoadedModuleNames.AddRange(
+			new string[]
+			{
+				// ... add any modules that your module loads dynamically here ...
+			}
+		);
+	}
+}
diff --git a/Engine/Plugins/Experimental/Animation/MLDeformer/DetailPoseModel/Source/DetailPoseModel/Private/DetailPoseModel.cpp b/Engine/Plugins/Experimental/Animation/MLDeformer/DetailPoseModel/Source/DetailPoseModel/Private/DetailPoseModel.cpp
new file mode 100644
--- /dev/null
+++ b/Engine/Plugins/Experimental/Animation/MLDeformer/DetailPoseModel/Source/DetailPoseModel/Private/DetailPoseModel.cpp
@@ -0,0 +1,117 @@
+// Copyright Epic Games, Inc. All Rights Reserved.
+#include "DetailPoseModel.h"
+#include "DetailPoseModelVizSettings.h"
+#include "DetailPoseModelInstance.h"
+#include "DetailPoseModelInputInfo.h"
+#include "MLDeformerComponent.h"
+#include "MLDeformerAsset.h"
+#include "UObject/AssetRegistryTagsContext.h"
+#include "UObject/UObjectGlobals.h"
+
+#include UE_INLINE_GENERATED_CPP_BY_NAME(DetailPoseModel)
+
+#define LOCTEXT_NAMESPACE "DetailPoseModel"
+
+// Implement our module.
+namespace UE::DetailPoseModel
+{
+	class DETAILPOSEMODEL_API FDetailPoseModelModule
+		: public IModuleInterface
+	{
+	};
+}
+IMPLEMENT_MODULE(UE::DetailPoseModel::FDetailPoseModelModule, DetailPoseModel)
+
+// Our log category for this model.
+DETAILPOSEMODEL_API DEFINE_LOG_CATEGORY(LogDetailPoseModel)
+
+//////////////////////////////////////////////////////////////////////////////
+
+UDetailPoseModel::UDetailPoseModel(const FObjectInitializer& ObjectInitializer)
+	: Super(ObjectInitializer)
+{
+	// We only want to use global mode in this model.
+	Mode = ENeuralMorphMode::Global;
+
+	// Create the visualization settings for this model.
+	// Never directly create one of the frameworks base classes such as the FMLDeformerMorphModelVizSettings as
+	// that can cause issues with detail customizations.
+#if WITH_EDITORONLY_DATA
+	SetVizSettings(ObjectInitializer.CreateEditorOnlyDefaultSubobject<UDetailPoseModelVizSettings>(this, TEXT("DetailPoseModelVizSettings")));
+#endif
+}
+
+UMLDeformerModelInstance* UDetailPoseModel::CreateModelInstance(UMLDeformerComponent* Component)
+{
+	return NewObject<UDetailPoseModelInstance>(Component);
+}
+
+UMLDeformerInputInfo* UDetailPoseModel::CreateInputInfo()
+{
+	return NewObject<UDetailPoseModelInputInfo>(this);
+}
+
+#if WITH_EDITORONLY_DATA
+	UAnimSequence* UDetailPoseModel::GetDetailPosesAnimSequence() const
+	{ 
+		return DetailPosesAnimSequence.LoadSynchronous();
+	}
+
+	UGeometryCache* UDetailPoseModel::GetDetailPosesGeomCache() const
+	{ 
+		return DetailPosesGeomCache.LoadSynchronous();
+	}
+#endif
+
+void UDetailPoseModel::Serialize(FArchive& Archive)
+{
+	// If we are saving and have raw deltas, strip the detail pose deltas from them, as we can just generate them when needed.
+	// This can keep the editor asset (non cooked) as small as possible.
+	TArray<FVector3f> DeltasBackup;
+	if (Archive.IsSaving() && GetInputInfo() && !GetMorphTargetDeltas().IsEmpty())
+	{
+		const UDetailPoseModelInputInfo* DetailPoseModelInputInfo = Cast<UDetailPoseModelInputInfo>(GetInputInfo());
+		check(DetailPoseModelInputInfo);
+		const int32 NumVertices = DetailPoseModelInputInfo->GetNumBaseMeshVertices();
+		if (NumVertices > 0)
+		{
+			const int32 NumMorphTargets = DetailPoseModelInputInfo->GetNumGlobalMorphTargets() + 1; // +1 for the means.
+			const TConstArrayView<FVector3f> MorphDeltasWithoutDetailPoses(GetMorphTargetDeltas().GetData(), NumVertices * NumMorphTargets);
+			TArray<FVector3f> StrippedDeltas(MorphDeltasWithoutDetailPoses);
+			DeltasBackup = GetMorphTargetDeltas();
+			SetMorphTargetDeltas(StrippedDeltas);
+		}
+	}
+
+	Super::Serialize(Archive);
+
+	// Recover the original deltas we had before we stripped out the detail pose deltas.
+	if (!DeltasBackup.IsEmpty())
+	{
+		SetMorphTargetDeltas(DeltasBackup);
+	}
+}
+
+void UDetailPoseModel::GetAssetRegistryTags(FAssetRegistryTagsContext Context) const
+{
+	Super::GetAssetRegistryTags(Context);
+
+	#if WITH_EDITORONLY_DATA
+		const FString DetailPosesAnimSequenceString = DetailPosesAnimSequence.ToSoftObjectPath().ToString();
+		const FString DetailPosesGeomCacheString = DetailPosesGeomCache.ToSoftObjectPath().ToString();
+		Context.AddTag(FAssetRegistryTag("MLDeformer.DetailPoseModel.DetailPosesAnimSequence", DetailPosesAnimSequenceString, FAssetRegistryTag::TT_Alphabetical));
+		Context.AddTag(FAssetRegistryTag("MLDeformer.DetailPoseModel.DetailPosesGeometryCache", DetailPosesGeomCacheString, FAssetRegistryTag::TT_Alphabetical));		
+	#endif
+
+	const UDetailPoseModelInputInfo* DetailPoseModelInputInfo = Cast<UDetailPoseModelInputInfo>(GetInputInfo());
+	if (DetailPoseModelInputInfo)
+	{
+		Context.AddTag(FAssetRegistryTag("MLDeformer.DetailPoseModel.Trained.NumGlobalMorphTargets", FString::FromInt(DetailPoseModelInputInfo->GetNumGlobalMorphTargets()), FAssetRegistryTag::TT_Numerical));
+	}
+
+	Context.AddTag(FAssetRegistryTag("MLDeformer.DetailPoseModel.BlendSpeed", FString::Printf(TEXT("%f"), BlendSpeed), FAssetRegistryTag::TT_Numerical));
+	Context.AddTag(FAssetRegistryTag("MLDeformer.DetailPoseModel.RGBRange", FString::Printf(TEXT("%f"), RBFRange), FAssetRegistryTag::TT_Numerical));
+	Context.AddTag(FAssetRegistryTag("MLDeformer.DetailPoseModel.UseRBFInterpolation", bUseRBF ? TEXT("True") : TEXT("False"), FAssetRegistryTag::TT_Alphabetical));
+}
+
+#undef LOCTEXT_NAMESPACE
diff --git a/Engine/Plugins/Experimental/Animation/MLDeformer/DetailPoseModel/Source/DetailPoseModel/Private/DetailPoseModelInputInfo.cpp b/Engine/Plugins/Experimental/Animation/MLDeformer/DetailPoseModel/Source/DetailPoseModel/Private/DetailPoseModelInputInfo.cpp
new file mode 100644
--- /dev/null
+++ b/Engine/Plugins/Experimental/Animation/MLDeformer/DetailPoseModel/Source/DetailPoseModel/Private/DetailPoseModelInputInfo.cpp
@@ -0,0 +1,10 @@
+// Copyright Epic Games, Inc. All Rights Reserved.
+
+#include "DetailPoseModelInputInfo.h"
+
+void UDetailPoseModelInputInfo::Reset()
+{
+	Super::Reset();
+
+	NumGlobalMorphTargets = 0;
+}
diff --git a/Engine/Plugins/Experimental/Animation/MLDeformer/DetailPoseModel/Source/DetailPoseModel/Private/DetailPoseModelInstance.cpp b/Engine/Plugins/Experimental/Animation/MLDeformer/DetailPoseModel/Source/DetailPoseModel/Private/DetailPoseModelInstance.cpp
new file mode 100644
--- /dev/null
+++ b/Engine/Plugins/Experimental/Animation/MLDeformer/DetailPoseModel/Source/DetailPoseModel/Private/DetailPoseModelInstance.cpp
@@ -0,0 +1,205 @@
+// Copyright Epic Games, Inc. All Rights Reserved.
+
+#include "DetailPoseModelInstance.h"
+#include "DetailPoseModel.h"
+#include "DetailPoseModelVizSettings.h"
+#include "DetailPoseModelInputInfo.h"
+#include "NeuralMorphNetwork.h"
+#include "Components/SkeletalMeshComponent.h"
+#include "Components/ExternalMorphSet.h"
+#include "Algo/MinElement.h"
+
+#if INTEL_ISPC
+	#include "DetailPoseModelPoseSearch.ispc.generated.h"
+#endif
+
+CSV_DECLARE_CATEGORY_MODULE_EXTERN(MLDEFORMERFRAMEWORK_API, MLDeformer);
+
+namespace 
+{
+	float CalculateDetailPoseDistance(const TConstArrayView<float> ValuesA, const TConstArrayView<float> ValuesB)
+	{
+		check(ValuesA.Num() == ValuesB.Num());
+
+		#if INTEL_ISPC
+			return ispc::CalcDistance(ValuesA.GetData(), ValuesB.GetData(), ValuesA.Num());
+		#else
+			float Distance = 0.0f;
+			const int32 NumValues = ValuesA.Num();
+			for (int32 Index = 0; Index < NumValues; ++Index)
+			{
+				const float Diff = ValuesA[Index] - ValuesB[Index];
+				Distance += Diff * Diff;
+			}
+
+			return Distance;
+		#endif
+	}
+
+	void CalculateDetailPoseDistances(const UDetailPoseModel* DetailPoseModel, const TConstArrayView<float> CurrentPoseValues, TArray<float>& OutDetailPoseDistances)
+	{
+		const TArray<FDetailPoseModelDetailPose>& DetailPoses = DetailPoseModel->GetDetailPoses();
+		for (int32 DetailPoseIndex = 0; DetailPoseIndex < DetailPoses.Num(); ++DetailPoseIndex)
+		{
+			const TConstArrayView<float> DetailPoseValues(DetailPoses[DetailPoseIndex].PoseValues);
+			OutDetailPoseDistances[DetailPoseIndex] = CalculateDetailPoseDistance(DetailPoseValues, CurrentPoseValues);
+		}
+	}
+
+	void UpdateWeight(float& OutWeight, float& InOutPrevWeight, float TargetWeight, float BlendSpeed)
+	{
+		const float BlendedWeight = BlendSpeed * TargetWeight + (1.0f - BlendSpeed) * InOutPrevWeight;
+		OutWeight = BlendedWeight;
+		InOutPrevWeight = BlendedWeight;
+	}
+
+	void UpdateRBFWeights(TArrayView<float> OutMorphWeights, TArrayView<float> SquaredDistances, TArrayView<float> InOutPrevWeights, float Sigma, float DetailPoseWeight, float BlendSpeed)
+	{
+		// NOTE: this method modifies the squared distances as well.
+
+		const int32 NumDetailPoses = OutMorphWeights.Num();
+		check(InOutPrevWeights.Num() == NumDetailPoses);
+		check(SquaredDistances.Num() == NumDetailPoses);
+		check(NumDetailPoses > 0);
+
+		const float MinD2 = *Algo::MinElement(SquaredDistances);
+		float SumWeights = 0.0f;
+		const float Sigma2 = FMath::Max(Sigma * Sigma, 1e-6f);
+		for (int32 Index = 0; Index < NumDetailPoses; ++Index)
+		{
+			float Weight = (SquaredDistances[Index] - MinD2) / Sigma2;
+			constexpr float CutOff = 3.0f;
+			if (Weight < CutOff)
+			{
+				Weight = FMath::Exp(-Weight);
+				SquaredDistances[Index] = Weight;
+				SumWeights += Weight;
+			}
+			else
+			{
+				SquaredDistances[Index] = 0.0f;
+			}
+		}
+
+		SumWeights = FMath::Max(SumWeights, 1e-6f);
+		for (int32 Index = 0; Index < NumDetailPoses; ++Index)
+		{
+			const float W = SquaredDistances[Index] / SumWeights * DetailPoseWeight;
+			UpdateWeight(OutMorphWeights[Index], InOutPrevWeights[Index], W, BlendSpeed);
+		}
+	}
+
+	void UpdateWeights(TArrayView<float> OutDetailPoseWeights, TConstArrayView<float> DistancePerDetailPose, TArrayView<float> InOutPrevDetailPoseWeights, float BlendSpeed, float ModelWeight)
+	{
+		// Find the detail pose index that has the smallest distance (to the current pose).
+		const float* MinDistancePtr = Algo::MinElement(DistancePerDetailPose);
+		const int32 MinDistanceDetailPoseIndex = MinDistancePtr - DistancePerDetailPose.GetData();
+
+		for (int32 DetailPoseIndex = 0; DetailPoseIndex < OutDetailPoseWeights.Num(); ++DetailPoseIndex)
+		{
+			const float TargetWeight = (MinDistanceDetailPoseIndex == DetailPoseIndex) ? ModelWeight : 0.0f;
+			UpdateWeight(OutDetailPoseWeights[DetailPoseIndex], InOutPrevDetailPoseWeights[DetailPoseIndex], TargetWeight, BlendSpeed); 
+		}
+	}
+
+	void UpdateDetailPoseWeights(const UDetailPoseModel* DetailPoseModel, TArrayView<float> OutDetailPoseWeights, TArrayView<float> DistancePerDetailPose, TArrayView<float> InOutPrevDetailPoseWeights, float BlendSpeed, float ModelWeight)
+	{
+		#if WITH_EDITOR
+			const UDetailPoseModelVizSettings* CastVizSettings = Cast<UDetailPoseModelVizSettings>(DetailPoseModel->GetVizSettings());
+			check(CastVizSettings);
+			const float DetailPoseWeight = CastVizSettings->GetDetailPoseWeight();
+		#else
+			const float DetailPoseWeight = 1.0f;
+		#endif
+
+		if (DetailPoseModel->GetUseRBFBlending())
+		{
+			UpdateRBFWeights(OutDetailPoseWeights, DistancePerDetailPose, InOutPrevDetailPoseWeights, DetailPoseModel->GetRBFRange(), ModelWeight * DetailPoseWeight, BlendSpeed);
+		}
+		else
+		{
+			UpdateWeights(OutDetailPoseWeights, DistancePerDetailPose, InOutPrevDetailPoseWeights, BlendSpeed, ModelWeight * DetailPoseWeight);
+		}
+	}
+}	// Anonymous namespace
+
+
+void UDetailPoseModelInstance::Init(USkeletalMeshComponent* SkelMeshComponent)
+{
+	Super::Init(SkelMeshComponent);
+
+	UDetailPoseModel* DetailPoseModel = Cast<UDetailPoseModel>(Model);
+	const int32 NumDetailPoses = DetailPoseModel->GetDetailPoses().Num();
+
+	DetailPoseDistances.Reset();
+	DetailPoseDistances.SetNumZeroed(NumDetailPoses);
+
+	DetailPosePrevWeights.Reset();
+	DetailPosePrevWeights.SetNumZeroed(NumDetailPoses);
+}
+
+void UDetailPoseModelInstance::Execute(float ModelWeight)
+{
+	TRACE_CPUPROFILER_EVENT_SCOPE(UDetailPoseModelInstance::Execute)
+	CSV_SCOPED_TIMING_STAT(MLDeformer, DetailPoseModelInstanceExecute);
+
+	if (!SkeletalMeshComponent)
+	{
+		return;
+	}
+
+	// Grab the weight data for this morph set.
+	// This could potentially fail if we are applying this deformer to the wrong skeletal mesh component.
+	const int32 LOD = SkeletalMeshComponent->GetPredictedLODLevel();
+	FExternalMorphSetWeights* WeightData = FindWeightData(LOD);
+	if (!WeightData)
+	{
+		return;
+	}
+
+	// Validate network output expectations.
+	UDetailPoseModel* DetailPoseModel = Cast<UDetailPoseModel>(Model);
+	UDetailPoseModelInputInfo* InputInfo = Cast<UDetailPoseModelInputInfo>(DetailPoseModel->GetInputInfo());
+	const int32 NumGlobalMorphs = InputInfo->GetNumGlobalMorphTargets();
+	UNeuralMorphNetwork* Network = DetailPoseModel->GetNeuralMorphNetwork();
+	const int32 NumNetworkOutputs = Network->GetNumMainOutputs();
+	if (!Network || Network->IsEmpty() || NumGlobalMorphs != NumNetworkOutputs || !NetworkInstance)
+	{
+		WeightData->ZeroWeights();
+		return;
+	}
+
+	// Calculate the weights the morph targets that the neural network calculates.
+	// This excludes the weights of the detail poses. We have concatenated the weights of the detail poses after those weights.
+	Super::Execute(ModelWeight);
+
+	// Now that we updated the weights of the morph targets that we generated during training, we have to deal with the
+	// weights of our detail poses. We want to calculate those, but we have to do some sanity checks first, and if there is an invalid state.
+	// In case there is some issue with the trained model (like if it hasn't been trained), we already zeroed the weights of the 
+	// detail poses as well.
+
+	// Now that we know the internal state is valid and we have the weight data, let's calculate and set the weights for the detail poses.
+	const int32 NumDetailPoses = WeightData->Weights.Num() - Network->GetNumMainOutputs() - 1;
+	if (NumDetailPoses > 0)
+	{
+		// Calculate the distances from the current pose to each detail pose.
+		// We can use these distances to find the closest pose, so we know which detail pose to blend in.
+		const TConstArrayView<float> CurrentPoseValues = NetworkInstance->GetInputs();
+		CalculateDetailPoseDistances(DetailPoseModel, CurrentPoseValues, DetailPoseDistances);
+
+		#if WITH_EDITOR
+			// Find the detail pose index that has the smallest distance (to the current pose).
+			const float* MinDistancePtr = Algo::MinElement(DetailPoseDistances);
+			BestDetailPoseIndex = MinDistancePtr - DetailPoseDistances.GetData();			  
+		#endif
+
+		// Now that we have calculated the distances, update the weights of each detail pose morph target.
+		// The weights of the detail poses are concatenated after the regular morph target weights.
+		// So we create an array view of just the detail pose weights here.
+		const float BlendSpeed = DetailPoseModel->GetBlendSpeed();
+		TArrayView<float> DetailPoseWeights(&WeightData->Weights[Network->GetNumMainOutputs() + 1], NumDetailPoses);	// +1 because of the first morph target being the vertex means and isn't part of the network outputs.
+		UpdateDetailPoseWeights(DetailPoseModel, DetailPoseWeights, DetailPoseDistances, DetailPosePrevWeights, BlendSpeed, ModelWeight);
+	}
+}
+
+
diff --git a/Engine/Plugins/Experimental/Animation/MLDeformer/DetailPoseModel/Source/DetailPoseModel/Private/DetailPoseModelPoseSearch.ispc b/Engine/Plugins/Experimental/Animation/MLDeformer/DetailPoseModel/Source/DetailPoseModel/Private/DetailPoseModelPoseSearch.ispc
new file mode 100644
--- /dev/null
+++ b/Engine/Plugins/Experimental/Animation/MLDeformer/DetailPoseModel/Source/DetailPoseModel/Private/DetailPoseModelPoseSearch.ispc
@@ -0,0 +1,14 @@
+// Copyright Epic Games, Inc. All Rights Reserved.
+
+export uniform float CalcDistance(uniform const float A[], uniform const float B[], uniform const int NumItems)
+{
+	uniform float Sum = 0.0f;
+    
+    foreach (Index = 0 ... NumItems)
+	{
+        const varying float Diff = A[Index] - B[Index];
+        Sum += reduce_add(Diff * Diff);
+    }
+    
+    return Sum;	
+}
diff --git a/Engine/Plugins/Experimental/Animation/MLDeformer/DetailPoseModel/Source/DetailPoseModel/Public/DetailPoseModel.h b/Engine/Plugins/Experimental/Animation/MLDeformer/DetailPoseModel/Source/DetailPoseModel/Public/DetailPoseModel.h
new file mode 100644
--- /dev/null
+++ b/Engine/Plugins/Experimental/Animation/MLDeformer/DetailPoseModel/Source/DetailPoseModel/Public/DetailPoseModel.h
@@ -0,0 +1,124 @@
+// Copyright Epic Games, Inc. All Rights Reserved.
+#pragma once
+
+#include "CoreTypes.h"
+#include "NeuralMorphModel.h"
+#include "DetailPoseModel.generated.h"
+
+// Declare our log category.
+DETAILPOSEMODEL_API DECLARE_LOG_CATEGORY_EXTERN(LogDetailPoseModel, Log, All);
+
+USTRUCT()
+struct DETAILPOSEMODEL_API FDetailPoseModelDetailPose
+{
+	GENERATED_BODY()
+
+	/** The values that represent the pose. These are the same as the neural network inputs that would represent this pose. */
+	UPROPERTY()
+	TArray<float> PoseValues;
+};
+
+/**
+ * The detail pose model for the ML Deformer.
+ * This model is inherited from the Neural Morph Model, but adds additional morph targets on top of this.
+ * The additional morph targets that are generated are morph targets that bring it to the ground truth as seen during training at specific poses.
+ * These special poses are called Detail Poses. We basically look at the deltas that we predict using the neural morph model, and put the remaining error
+ * between that and the ground truth into the morph target at the given pose.
+ * There can be multiple detail poses, and each of them generates a new morph target.
+ */
+UCLASS()
+class DETAILPOSEMODEL_API UDetailPoseModel 
+	: public UNeuralMorphModel
+{
+	GENERATED_BODY()
+
+public:
+	UDetailPoseModel(const FObjectInitializer& ObjectInitializer);
+
+	// UObject overrides.
+	void Serialize(FArchive& Archive) override final;
+	void GetAssetRegistryTags(FAssetRegistryTagsContext Context) const override final;
+	// ~END UObject overrides.
+
+	// UMLDeformerModel overrides.
+	FString GetDisplayName() const override final						{ return "Detail Pose Model"; }
+	UMLDeformerModelInstance* CreateModelInstance(UMLDeformerComponent* Component) override final;
+	UMLDeformerInputInfo* CreateInputInfo() override final;
+	// ~END UMLDeformerModel overrides.
+
+	// UNeuralMorphModel overrides.
+	bool SupportsGlobalModeOnly() const override final					{ return true; }
+	// ~END UNeuralMorphModel overrides.
+
+	const TArray<FDetailPoseModelDetailPose>& GetDetailPoses() const	{ return DetailPoses; }
+	TArray<FDetailPoseModelDetailPose>& GetDetailPoses()				{ return DetailPoses; }
+	float GetBlendSpeed() const											{ return BlendSpeed; }
+	bool GetUseRBFBlending() const										{ return bUseRBF; }
+	float GetRBFRange() const											{ return RBFRange; }
+
+	void SetUseRBFBlending(bool bUse)									{ bUseRBF = bUse; }
+	void SetBlendSpeed(float Speed)										{ check(Speed >= 0.0f && Speed <= 1.0f); BlendSpeed = Speed; }
+	void SetRBFRange(float Range)										{ check(Range >= 0.0f); RBFRange = Range; }
+
+#if WITH_EDITORONLY_DATA
+	UAnimSequence* GetDetailPosesAnimSequence() const;
+	UGeometryCache* GetDetailPosesGeomCache() const;
+
+	static FName GetDetailPosesAnimSequencePropertyName()				{ return GET_MEMBER_NAME_CHECKED(UDetailPoseModel, DetailPosesAnimSequence); }
+	static FName GetDetailPosesGeomCachePropertyName()					{ return GET_MEMBER_NAME_CHECKED(UDetailPoseModel, DetailPosesGeomCache); }
+	static FName GetBlendSpeedPropertyName()							{ return GET_MEMBER_NAME_CHECKED(UDetailPoseModel, BlendSpeed); }
+	static FName GetUseRBFBlendingPropertyName()						{ return GET_MEMBER_NAME_CHECKED(UDetailPoseModel, bUseRBF); }
+	static FName GetRBFRangePropertyName()								{ return GET_MEMBER_NAME_CHECKED(UDetailPoseModel, RBFRange); }
+#endif
+
+private:
+	/** The detail poses. This contains the pose values we use at runtime. */
+	UPROPERTY()
+	TArray<FDetailPoseModelDetailPose> DetailPoses;
+
+#if WITH_EDITORONLY_DATA
+	/**
+	 * For certain poses the ML Deformer might not reconstruct important details such as specific cloth folding patterns.
+	 * To improve the results near these important poses, you can provide a set of these poses.
+	 * The poses are provided similar to the training input data, with an animation sequence and geometry cache pair where each
+	 * frame contains a pose you want to preserve details in.
+	 * The more poses you add, the higher the memory usage. Each pose will generate an additional morph target.
+	 * If you do not add any detail poses, the model will behave the same as the Neural Morph Model in global mode.
+	 * Please keep the frame rate and frame count the same for both the anim sequence and geometry cache.
+	 */
+	UPROPERTY(EditAnywhere, Category = "Detail Poses")
+	TSoftObjectPtr<UAnimSequence> DetailPosesAnimSequence;
+
+	/**
+	 * For certain poses the ML Deformer might not reconstruct important details such as specific cloth folding patterns.
+	 * To improve the results near these important poses, you can provide a set of these poses.
+	 * The poses are provided similar to the training input data, with an animation sequence and geometry cache pair where each
+	 * frame contains a pose you want to preserve details in.
+	 * The more poses you add, the higher the memory usage. Each pose will generate an additional morph target.
+	 * If you do not add any detail poses, the model will behave the same as the Neural Morph Model in global mode.
+	 * Please keep the frame rate and frame count the same for both the anim sequence and geometry cache.
+	 */
+	UPROPERTY(EditAnywhere, Category = "Detail Poses")
+	TSoftObjectPtr<UGeometryCache> DetailPosesGeomCache;
+#endif
+
+	/**
+	 * The speed at which the detail poses are blend in. 
+	 * Higher values make it blend in faster.
+	 * A value of 0.0 would disable the detail poses from being calculated.
+	 * A value of 1.0 would disable blending and instantly switch the active detail pose weight.
+	 */
+	UPROPERTY(EditAnywhere, Category = "Detail Poses", meta = (ClampMin="0.0", ClampMax="1.0"))
+	float BlendSpeed = 0.3f;
+
+	/**
+	 * Use RBF (Radial Basis Functions) to interpolate?
+	 * This will produce higher quality blends between detail poses, at the cost of runtime CPU performance.
+	 */
+	UPROPERTY(EditAnywhere, DisplayName = "Use RBF Blending", Category = "Detail Poses")
+	bool bUseRBF = true;
+
+	/** The range to blend detail poses. Larger values will blend more detail poses together but also result in slower GPU performance. */
+	UPROPERTY(EditAnywhere, DisplayName = "RBF Range", Category = "Detail Poses", meta = (ClampMin = "0", EditCondition = "bUseRBF"))
+	float RBFRange = 1.0f;
+};
diff --git a/Engine/Plugins/Experimental/Animation/MLDeformer/DetailPoseModel/Source/DetailPoseModel/Public/DetailPoseModelInputInfo.h b/Engine/Plugins/Experimental/Animation/MLDeformer/DetailPoseModel/Source/DetailPoseModel/Public/DetailPoseModelInputInfo.h
new file mode 100644
--- /dev/null
+++ b/Engine/Plugins/Experimental/Animation/MLDeformer/DetailPoseModel/Source/DetailPoseModel/Public/DetailPoseModelInputInfo.h
@@ -0,0 +1,48 @@
+// Copyright Epic Games, Inc. All Rights Reserved.
+#pragma once
+
+#include "CoreTypes.h"
+#include "NeuralMorphInputInfo.h"
+#include "DetailPoseModelInputInfo.generated.h"
+
+namespace UE::DetailPoseModel
+{ 
+	class FDetailPoseEditorModel;
+}
+
+/**
+ * The input info class for the Detail Pose Model.
+ * We store values that were used during training. These are often used during inference.
+ */
+UCLASS()
+class DETAILPOSEMODEL_API UDetailPoseModelInputInfo
+	: public UNeuralMorphInputInfo
+{
+	GENERATED_BODY()
+
+public:
+	friend class UE::DetailPoseModel::FDetailPoseEditorModel;
+
+	// FMLDeformerInputInfo overrides.
+	void Reset() override final;
+	// ~END FMLDeformerInputInfo overrides.
+
+	/**
+	 * Get the number of morph targets that the model should output.
+	 * This is basically the number of outputs of our neural network. Please keep in mind that the number of morph targets
+	 * that are actually used is higher than this number. The number returned by this method is really what we entered as 
+	 * number of morph targets in the UI.
+	 * We do generate an extra morph target that holds the mean values as well, and next to that there is a morph target
+	 * for each detail pose. So please keep in mind this is just the number of morph targets we entered in the UI before we pressed 
+	 * the Train button, and that it is equal to the number of neural network outputs.
+	 */
+	int32 GetNumGlobalMorphTargets() const			{ return NumGlobalMorphTargets; }
+
+private:
+	void SetNumGlobalMorphTargets(int32 Num)		{ NumGlobalMorphTargets = Num; }
+
+private:
+	/** The number of morph targets the training process generated. This excludes the morph target that contains the means. */
+	UPROPERTY()
+	int32 NumGlobalMorphTargets = 0;
+};
diff --git a/Engine/Plugins/Experimental/Animation/MLDeformer/DetailPoseModel/Source/DetailPoseModel/Public/DetailPoseModelInstance.h b/Engine/Plugins/Experimental/Animation/MLDeformer/DetailPoseModel/Source/DetailPoseModel/Public/DetailPoseModelInstance.h
new file mode 100644
--- /dev/null
+++ b/Engine/Plugins/Experimental/Animation/MLDeformer/DetailPoseModel/Source/DetailPoseModel/Public/DetailPoseModelInstance.h
@@ -0,0 +1,43 @@
+// Copyright Epic Games, Inc. All Rights Reserved.
+
+#pragma once
+
+#include "NeuralMorphModelInstance.h"
+#include "DetailPoseModelInstance.generated.h"
+
+class USkeletalMeshComponent;
+
+/**
+ * The model instance class of the Detail Pose Model.
+ * This contains the code that calculates the morph target weights and which calculates
+ * which detail pose should be blend in.
+ */
+UCLASS()
+class DETAILPOSEMODEL_API UDetailPoseModelInstance
+	: public UNeuralMorphModelInstance
+{
+	GENERATED_BODY()
+
+public:
+	// UMLDeformerModelInstance overrides.
+	void Init(USkeletalMeshComponent* SkelMeshComponent) override final;
+	void Execute(float ModelWeight) override final;
+	// ~END UMLDeformerModelInstance overrides.
+
+	/**
+	 * Get the best matching detail pose, compared to the character's current pose as defined by the skeletal mesh component we are linked to.
+	 * The integer that is returned represents the frame number inside the detail pose model's detail pose geometry cache.
+	 * @return The index of the frame number in the detail pose geometry cache that best matches our current pose.
+	 */
+	int32 GetBestDetailPoseIndex() const		{ return BestDetailPoseIndex; }
+
+private:
+	/** The distance value between the current current pose and each detail pose. */
+	TArray<float> DetailPoseDistances;
+
+	/** The detail pose previous frame weights, one for each detail pose. */
+	TArray<float> DetailPosePrevWeights;
+
+	/** The detail pose index that currently is the one with closest distance. This is the frame number inside the detail pose geometry cache. */
+	int32 BestDetailPoseIndex = -1;
+};
diff --git a/Engine/Plugins/Experimental/Animation/MLDeformer/DetailPoseModel/Source/DetailPoseModel/Public/DetailPoseModelVizSettings.h b/Engine/Plugins/Experimental/Animation/MLDeformer/DetailPoseModel/Source/DetailPoseModel/Public/DetailPoseModelVizSettings.h
new file mode 100644
--- /dev/null
+++ b/Engine/Plugins/Experimental/Animation/MLDeformer/DetailPoseModel/Source/DetailPoseModel/Public/DetailPoseModelVizSettings.h
@@ -0,0 +1,42 @@
+// Copyright Epic Games, Inc. All Rights Reserved.
+
+#pragma once
+
+#include "CoreTypes.h"
+#include "NeuralMorphModelVizSettings.h"
+#include "DetailPoseModelVizSettings.generated.h"
+
+/**
+ * The visualization settings specific to the Detail Pose Model.
+ */
+UCLASS()
+class DETAILPOSEMODEL_API UDetailPoseModelVizSettings
+	: public UNeuralMorphModelVizSettings
+{
+	GENERATED_BODY()
+
+public:
+#if WITH_EDITORONLY_DATA
+	float GetDetailPoseWeight() const						{ return DetailPoseWeight; }
+	void SetDetailPoseWeight(float Factor)					{ DetailPoseWeight = FMath::Clamp(Factor, 0.0f, 1.0f); }
+
+	bool GetDrawDetailPose() const							{ return bDrawDetailPose; }
+	void SetDrawDetailPose(bool bDraw)						{ bDrawDetailPose = bDraw; }
+
+	static FName GetDetailPoseWeightPropertyName()			{ return GET_MEMBER_NAME_CHECKED(UDetailPoseModelVizSettings, DetailPoseWeight); }
+	static FName GetDrawDetailPosePropertyName()			{ return GET_MEMBER_NAME_CHECKED(UDetailPoseModelVizSettings, bDrawDetailPose); }
+
+private:
+	/** 
+	 * The weight that specifies how much of the detail poses should be added.
+	 * When this value is set to 0 it means the detail poses are not visually active.
+	 * When set to 1, they are fully active.
+	 */
+	UPROPERTY(EditAnywhere, Category = "Live Settings|Detail Poses", meta = (ClampMin = "0.0", ClampMax = "1.0"))
+	float DetailPoseWeight = 1.0f;
+
+	/** Do we want to draw the closest matching detail pose? */
+	UPROPERTY(EditAnywhere, Category = "Live Settings|Detail Poses")
+	bool bDrawDetailPose = true;
+#endif	// WITH_EDITORONLY_DATA
+};
diff --git a/Engine/Plugins/Experimental/Animation/MLDeformer/DetailPoseModel/Source/DetailPoseModelEditor/DetailPoseModelEditor.Build.cs b/Engine/Plugins/Experimental/Animation/MLDeformer/DetailPoseModel/Source/DetailPoseModelEditor/DetailPoseModelEditor.Build.cs
new file mode 100644
--- /dev/null
+++ b/Engine/Plugins/Experimental/Animation/MLDeformer/DetailPoseModel/Source/DetailPoseModelEditor/DetailPoseModelEditor.Build.cs
@@ -0,0 +1,47 @@
+// Copyright Epic Games, Inc. All Rights Reserved.
+
+using System.IO;
+
+namespace UnrealBuildTool.Rules
+{
+	public class DetailPoseModelEditor : ModuleRules
+	{
+		public DetailPoseModelEditor(ReadOnlyTargetRules Target) : base(Target)
+		{
+			PublicDependencyModuleNames.AddRange(
+				new string[]
+				{
+					"Core"
+				}
+			);
+
+			PrivateDependencyModuleNames.AddRange(
+				new string[]
+				{
+					"InputCore",
+					"UnrealEd",
+					"CoreUObject",
+					"Engine",
+					"Slate",
+					"SlateCore",
+					"Projects",
+					"MLDeformerFramework",
+					"MLDeformerFrameworkEditor",
+					"DetailPoseModel",
+					"NeuralMorphModel",
+					"NeuralMorphModelEditor",
+					"PropertyEditor",
+					"ToolWidgets",
+					"DeveloperSettings",
+					"GeometryCache"
+				}
+			);
+
+			PrivateIncludePathModuleNames.AddRange(
+				new string[]
+				{
+				}
+			);
+		}
+	}
+}
diff --git a/Engine/Plugins/Experimental/Animation/MLDeformer/DetailPoseModel/Source/DetailPoseModelEditor/Private/DetailPoseEditorModel.cpp b/Engine/Plugins/Experimental/Animation/MLDeformer/DetailPoseModel/Source/DetailPoseModelEditor/Private/DetailPoseEditorModel.cpp
new file mode 100644
--- /dev/null
+++ b/Engine/Plugins/Experimental/Animation/MLDeformer/DetailPoseModel/Source/DetailPoseModelEditor/Private/DetailPoseEditorModel.cpp
@@ -0,0 +1,448 @@
+// Copyright Epic Games, Inc. All Rights Reserved.
+
+#include "DetailPoseEditorModel.h"
+#include "DetailPoseModel.h"
+#include "DetailPoseModelInstance.h"
+#include "DetailPoseModelInputInfo.h"
+#include "DetailPoseModelEditorStyle.h"
+#include "DetailPoseModelEditorActor.h"
+#include "DetailPoseModelVizSettings.h"
+#include "MLDeformerComponent.h"
+#include "MLDeformerGeomCacheActor.h"
+#include "MLDeformerGeomCacheSampler.h"
+#include "MLDeformerEditorToolkit.h"
+#include "NeuralMorphNetwork.h"
+#include "NeuralMorphEditorModel.h"
+#include "Components/SkeletalMeshComponent.h"
+#include "Animation/DebugSkelMeshComponent.h"
+#include "Misc/ScopedSlowTask.h"
+#include "IPersonaPreviewScene.h"
+#include "GeometryCacheComponent.h"
+#include "GeometryCache.h"
+#include "Async/ParallelFor.h"
+
+#define LOCTEXT_NAMESPACE "DetailPoseEditorModel"
+
+namespace UE::DetailPoseModel
+{
+	using namespace UE::MLDeformer;
+
+	FMLDeformerEditorModel* FDetailPoseEditorModel::MakeInstance()
+	{
+		return new FDetailPoseEditorModel();
+	}
+
+	FMLDeformerTrainingInputAnim* FDetailPoseEditorModel::GetTrainingInputAnim(int32 Index) const
+	{
+		if (Index < GetNumTrainingInputAnims())
+		{
+			return FNeuralMorphEditorModel::GetTrainingInputAnim(Index);
+		}
+		else
+		{
+			// Output the detail poses animation.
+			check(Index == GetNumTrainingInputAnims());
+			if (!DetailPosesAnim.IsValid())
+			{
+				DetailPosesAnim = MakeUnique<FMLDeformerGeomCacheTrainingInputAnim>();
+			}
+			DetailPosesAnim->SetAnimSequence(GetDetailPoseModel()->GetDetailPosesAnimSequence());
+			DetailPosesAnim->SetGeometryCache(GetDetailPoseModel()->GetDetailPosesGeomCache());
+			return DetailPosesAnim.Get();
+		}
+	}
+
+	void FDetailPoseEditorModel::UpdateTimelineTrainingAnimList()
+	{
+		TArray<TSharedPtr<FMLDeformerTrainingInputAnimName>> NameList;
+		
+		// Build the list of names, based on the training inputs.
+		const int32 NumAnims = GetNumTrainingInputAnims();
+		for (int32 AnimIndex = 0; AnimIndex < NumAnims; ++AnimIndex)
+		{
+			const FMLDeformerTrainingInputAnim* Anim = GetTrainingInputAnim(AnimIndex);
+			if (Anim)
+			{
+				if (Anim->IsValid())
+				{
+					TSharedPtr<FMLDeformerTrainingInputAnimName> AnimName = MakeShared<FMLDeformerTrainingInputAnimName>();
+					AnimName->TrainingInputAnimIndex = AnimIndex;
+					AnimName->Name = FString::Printf(TEXT("[#%d] %s"), AnimIndex, *Anim->GetAnimSequence()->GetName());
+					NameList.Emplace(AnimName);
+				}
+			}
+		}
+
+		const FMLDeformerTrainingInputAnim* Anim = GetTrainingInputAnim(NumAnims);
+		if (Anim && Anim->IsValid())
+		{
+			TSharedPtr<FMLDeformerTrainingInputAnimName> AnimName = MakeShared<FMLDeformerTrainingInputAnimName>();
+			AnimName->TrainingInputAnimIndex = NumAnims;
+			AnimName->Name = FString::Printf(TEXT("[Detail Poses] %s"), *Anim->GetAnimSequence()->GetName());	
+			NameList.Emplace(AnimName);
+		}
+
+		SetTimelineAnimNames(NameList);
+	}
+
+	void FDetailPoseEditorModel::OnPropertyChanged(FPropertyChangedEvent& PropertyChangedEvent)
+	{
+		FNeuralMorphEditorModel::OnPropertyChanged(PropertyChangedEvent);
+
+		const FProperty* Property = PropertyChangedEvent.Property;
+		if (Property == nullptr)
+		{
+			return;
+		}
+
+		if (Property->GetFName() == UMLDeformerGeomCacheModel::GetTrainingInputAnimsPropertyName() || PropertyChangedEvent.GetMemberPropertyName() == TEXT("TrainingInputAnims"))
+		{
+			if (PropertyChangedEvent.ChangeType == EPropertyChangeType::ArrayRemove || PropertyChangedEvent.ChangeType == EPropertyChangeType::ArrayClear)
+			{
+				if (GetDetailPoseModel()->GetTrainingInputAnims().IsEmpty() && DetailPosesAnim.IsValid())
+				{
+					SetActiveTrainingInputAnimIndex(0);
+					TriggerInputAssetChanged();
+				}
+			}
+		}
+		else if (Property->GetFName() == UDetailPoseModel::GetDetailPosesAnimSequencePropertyName() || Property->GetFName() == UDetailPoseModel::GetDetailPosesGeomCachePropertyName() || Property->GetFName() == UNeuralMorphModel::GetClampMorphTargetWeightsPropertyName())
+		{
+			TriggerInputAssetChanged();
+
+			// Empty the detail poses deltas so they get regenerated.
+			// And then reinitialize the engine morph targets to include the modified detail pose morphs.
+			DetailPosesDeltas.Empty();
+			InitEngineMorphTargets(GetMorphModel()->GetMorphTargetDeltas());
+		}
+		else if (Property->GetFName() == UDetailPoseModelVizSettings::GetDrawDetailPosePropertyName())
+		{
+			UpdateActorVisibility();
+		}
+	}
+
+	void FDetailPoseEditorModel::OnPostTraining(ETrainingResult TrainingResult, bool bUsePartiallyTrainedWhenAborted)
+	{
+		// On successful training, or when we aborted but still want to use the currently trained network,
+		// empty the detail poses deltas so they get regenerated.
+		// We need to regenerate these detail pose deltas as the trained neural network has changed.
+		if (TrainingResult == ETrainingResult::Success || (TrainingResult == ETrainingResult::Aborted && bUsePartiallyTrainedWhenAborted))
+		{
+			DetailPosesDeltas.Empty();
+		}
+
+		FNeuralMorphEditorModel::OnPostTraining(TrainingResult, bUsePartiallyTrainedWhenAborted);
+	}
+
+	void FDetailPoseEditorModel::CalculateDetailPoseDeltas(TArray<FVector3f>& OutDeltas, TArray<FDetailPoseModelDetailPose>& OutDetailPoses)
+	{
+		OutDeltas.Empty();
+		OutDetailPoses.Empty();
+
+		const UGeometryCache* DetailPosesGeomCache = GetDetailPoseModel()->GetDetailPosesGeomCache();
+		const UAnimSequence* DetailPosesAnimSequence= GetDetailPoseModel()->GetDetailPosesAnimSequence();
+		if (!DetailPosesGeomCache || !DetailPosesAnimSequence)
+		{
+			return;
+		}
+
+		const int32 NumDetailPoses = DetailPosesGeomCache->GetEndFrame() - DetailPosesGeomCache->GetStartFrame() + 1;
+		if (NumDetailPoses == 0)
+		{
+			return;
+		}
+
+		UDetailPoseModel* DetailPoseModel = GetDetailPoseModel();
+		UNeuralMorphNetwork* MorphNetwork = DetailPoseModel->GetNeuralMorphNetwork();
+		if (!MorphNetwork)
+		{
+			return;
+		}
+
+		FScopedSlowTask Task(NumDetailPoses, LOCTEXT("CalculateDetailPoseDeltasProgress", "Calculating detail pose morph deltas"));
+		Task.MakeDialog(false);	
+
+		// Calculate ground truth deltas.
+		FMLDeformerGeomCacheSampler Sampler;
+		Sampler.Init(this, GetNumTrainingInputAnims());	// We go one past the number of training input anims here, which will give us the Detail Poses Anim. See GetTrainingInputAnim.
+		Sampler.SetVertexDeltaSpace(EVertexDeltaSpace::PreSkinning);
+		Sampler.SetSkinningMode(GetDetailPoseModel()->GetSkinningMode());
+
+		// Create and setup the inference object.
+		UDetailPoseModelInstance* ModelInstance = NewObject<UDetailPoseModelInstance>();
+		ModelInstance->SetModel(Model);
+		USkeletalMeshComponent* SkelMeshComponent = NewObject<USkeletalMeshComponent>(ModelInstance);
+		SkelMeshComponent->SetSkeletalMesh(Model->GetSkeletalMesh());
+		ModelInstance->Init(SkelMeshComponent);
+		ModelInstance->PostMLDeformerComponentInit();
+
+		// Get access to the network inputs.
+		check(Sampler.GetNumFloatsPerCurve() == 1);
+		check(Model->GetNumFloatsPerCurve() == 1);
+		const int64 NumDeformerAssetInputs = Model->GetInputInfo()->CalcNumNeuralNetInputs(Model->GetNumFloatsPerBone(), Model->GetNumFloatsPerCurve());
+		const int32 NumNeuralNetworkInputs = MorphNetwork->GetNumMainInputs();
+		check(NumDeformerAssetInputs == NumNeuralNetworkInputs);
+
+		// Calculate the deltas for each detail pose.
+		const int32 NumVerts = Model->GetNumBaseMeshVerts();
+		OutDeltas.SetNumZeroed(NumVerts * NumDetailPoses);
+		OutDetailPoses.Reset();
+		OutDetailPoses.SetNum(NumDetailPoses);
+		for (int32 DetailPoseIndex = 0; DetailPoseIndex < NumDetailPoses; ++DetailPoseIndex)
+		{
+			Sampler.Sample(DetailPoseIndex);
+
+			TArrayView<FVector3f> DetailPoseDeltas(&OutDeltas[DetailPoseIndex * NumVerts], NumVerts);
+			CalculateDetailPoseDeltas(DetailPoseIndex, Sampler, ModelInstance, DetailPoseDeltas, OutDetailPoses[DetailPoseIndex]);
+			Task.EnterProgressFrame();
+		}
+
+		SkelMeshComponent->ConditionalBeginDestroy();
+		ModelInstance->ConditionalBeginDestroy();
+	}
+
+	void FDetailPoseEditorModel::CalculateDetailPoseDeltas(
+		int32 DetailPoseIndex,
+		FMLDeformerGeomCacheSampler& Sampler,
+		UDetailPoseModelInstance* ModelInstance,
+		TArrayView<FVector3f> OutDeltas,
+		FDetailPoseModelDetailPose& OutDetailPose)
+	{
+		// Calculate ground truth deltas.
+		const TArray<float>& GroundTruthDeltas = Sampler.GetVertexDeltas();
+
+		// Init the bone inputs.
+		UDetailPoseModel* DetailPoseModel = GetDetailPoseModel();
+		const TArray<float>& BoneInputFloats = Sampler.GetBoneRotations();
+		const TArray<float>& CurveInputFloats= Sampler.GetCurveValues();
+		const int32 NumBoneFloats = BoneInputFloats.Num();
+
+		UNeuralMorphNetworkInstance* NetworkInstance = ModelInstance->GetNetworkInstance();
+		TArrayView<float> NetworkInputBuffer = NetworkInstance->GetInputs();
+		UNeuralMorphNetwork* MorphNetwork = DetailPoseModel->GetNeuralMorphNetwork();
+		check(NumBoneFloats + CurveInputFloats.Num() == MorphNetwork->GetNumMainInputs());
+
+		const TConstArrayView<float> Means = MorphNetwork->GetInputMeans();
+		const TConstArrayView<float> Stds = MorphNetwork->GetInputStds();
+		for (int32 Index = 0; Index < NumBoneFloats; ++Index)
+		{
+			NetworkInputBuffer[Index] = (BoneInputFloats[Index] - Means[Index]) / Stds[Index];
+		}
+
+		// Init the curve inputs.
+		for (int32 Index = 0; Index < CurveInputFloats.Num(); ++Index)
+		{
+			const int32 CurveInputIndex = NumBoneFloats + Index;
+			NetworkInputBuffer[CurveInputIndex] = (CurveInputFloats[Index] - Means[CurveInputIndex]) / Stds[CurveInputIndex];
+		}
+
+		// Copy the input values as the pose values.
+		OutDetailPose.PoseValues = NetworkInputBuffer;
+
+		// Now that we initialized the inputs, run inference.
+		NetworkInstance->Run();
+
+		// Clamp the output weights if clamping is enabled.
+		TArray<float> MorphWeights(NetworkInstance->GetOutputs());
+		if (DetailPoseModel->IsMorphWeightClampingEnabled())
+		{
+			DetailPoseModel->ClampMorphTargetWeights(MorphWeights);
+		}
+
+		// Blend the raw vertex deltas with the morph weights that the neural network has output, to get the predicted deltas.
+		const int32 NumVertices = Model->GetNumBaseMeshVerts();
+		const TConstArrayView<FVector3f> RawMorphDeltas(DetailPoseModel->GetMorphTargetDeltas().GetData(), NumVertices * (MorphWeights.Num() + 1));
+		check(RawMorphDeltas.Num() % NumVertices == 0);
+		check(Sampler.GetNumImportedVertices() == NumVertices);
+		TArray<FVector3f> PredictedDeltas;
+		PredictedDeltas.SetNumZeroed(NumVertices);
+		const int32 NumMorphTargets = RawMorphDeltas.Num() / NumVertices;
+		check(MorphWeights.Num() + 1 == NumMorphTargets);
+
+		ParallelFor(TEXT("CalcPredictedDetailPoseDeltas"), NumVertices, 500, [&PredictedDeltas, NumMorphTargets, NumVertices, &RawMorphDeltas, &MorphWeights](int32 VertexIndex)
+		{
+			for (int32 MorphIndex = 0; MorphIndex < NumMorphTargets; ++MorphIndex)
+			{
+				const int32 MorphVertexOffset = MorphIndex * NumVertices;
+				const float MorphWeight = MorphIndex > 0 ? MorphWeights[MorphIndex - 1] : 1.0f;	// The first morph target contains the means that we should add, use a weight of 1 there.
+				PredictedDeltas[VertexIndex] += RawMorphDeltas[MorphVertexOffset + VertexIndex] * MorphWeight;
+			}
+		});
+
+		// Calculate and output the difference between the ground truth and predicted deltas.
+		check(PredictedDeltas.Num() * 3 == Sampler.GetVertexDeltas().Num());
+		check(PredictedDeltas.Num() == OutDeltas.Num());
+		check(OutDeltas.Num() == NumVertices);
+		for (int32 VertexIndex = 0; VertexIndex < NumVertices; ++VertexIndex)
+		{
+			const int32 GroundTruthVertexOffset = VertexIndex * 3;
+			const FVector3f GroundTruthDelta(GroundTruthDeltas[GroundTruthVertexOffset], GroundTruthDeltas[GroundTruthVertexOffset+1], GroundTruthDeltas[GroundTruthVertexOffset+2]);
+			const FVector3f PredictedDelta = PredictedDeltas[VertexIndex];
+			OutDeltas[VertexIndex] = GroundTruthDelta - PredictedDelta;	
+		}
+	}
+
+	FDetailPoseModelEditorActor* FDetailPoseEditorModel::CreateDetailPoseActor(UWorld* World) const
+	{
+		const FLinearColor LabelColor = FDetailPoseModelEditorStyle::Get().GetColor("DetailPoseModel.EditorActor.LabelColor");
+		const FLinearColor WireframeColor = FDetailPoseModelEditorStyle::Get().GetColor("DetailPoseModel.EditorActor.WireframeColor");
+		const FName ActorName = FName("DetailPoseActor");
+		const FText LabelText = LOCTEXT("DetailPoseActorLabelText", "Detail Pose");
+		const int32 ActorID = ActorID_DetailPoseActor;
+
+		FActorSpawnParameters SpawnParams;
+		SpawnParams.Name = MakeUniqueObjectName(World, AActor::StaticClass(), ActorName);
+		AActor* Actor = World->SpawnActor<AActor>(SpawnParams);
+		Actor->SetFlags(RF_Transient);
+		
+		// Create the Geometry Cache Component.
+		UGeometryCacheComponent* GeomCacheComponent = NewObject<UGeometryCacheComponent>(Actor);
+		GeomCacheComponent->RegisterComponent();
+		GeomCacheComponent->SetOverrideWireframeColor(true);
+		GeomCacheComponent->SetWireframeOverrideColor(WireframeColor);
+		GeomCacheComponent->MarkRenderStateDirty();
+		GeomCacheComponent->SetVisibility(false);
+		Actor->SetRootComponent(GeomCacheComponent);
+		
+		// Create the editor actor.
+		UE::MLDeformer::FMLDeformerEditorActor::FConstructSettings Settings;
+		Settings.Actor = Actor;
+		Settings.TypeID = ActorID;
+		Settings.LabelColor = LabelColor;
+		Settings.LabelText = LabelText;
+		Settings.bIsTrainingActor = false;
+
+		FDetailPoseModelEditorActor* NewActor = new FDetailPoseModelEditorActor(Settings);
+		NewActor->SetGeometryCacheComponent(GeomCacheComponent);
+		return NewActor;
+	}
+
+	UMLDeformerComponent* FDetailPoseEditorModel::GetTestMLDeformerComponent() const
+	{
+		return FindMLDeformerComponent(UE::MLDeformer::ActorID_Test_MLDeformed);
+	}
+
+	void FDetailPoseEditorModel::UpdateDetailPoseActor(FDetailPoseModelEditorActor& Actor) const
+	{
+		using namespace UE::MLDeformer;
+
+		const UDetailPoseModel* DetailPoseModel = GetDetailPoseModel();
+		if (!DetailPoseModel)
+		{
+			return;
+		}
+
+		float MaxOffset = 0.0f;
+		for (FMLDeformerEditorActor* EditorActor : EditorActors)
+		{
+			if (EditorActor && 
+				EditorActor->GetTypeID() != ActorID_DetailPoseActor &&
+				EditorActor->IsVisible()) 
+			{
+				MaxOffset = FMath::Max(EditorActor->GetMeshOffsetFactor(), MaxOffset);
+			}
+		}
+
+		Actor.SetMeshOffsetFactor(MaxOffset + 1.0f);
+
+		UGeometryCache* GeomCache = DetailPoseModel->GetDetailPosesGeomCache();
+		Actor.SetGeometryCache(GeomCache);
+
+		const UMLDeformerComponent* MLDeformerComponent = GetTestMLDeformerComponent();
+		Actor.SetTrackedComponent(MLDeformerComponent);
+	}
+
+	void FDetailPoseEditorModel::InitInputInfo(UMLDeformerInputInfo* InputInfo)
+	{
+		FNeuralMorphEditorModel::InitInputInfo(InputInfo);
+
+		UDetailPoseModelInputInfo* CastInputInfo = Cast<UDetailPoseModelInputInfo>(InputInfo);
+		if (!CastInputInfo)
+		{
+			return;
+		}
+
+		CastInputInfo->SetNumGlobalMorphTargets(GetDetailPoseModel()->GetGlobalNumMorphs());
+	}
+
+	void FDetailPoseEditorModel::CreateActors(const TSharedRef<IPersonaPreviewScene>& InPersonaPreviewScene)
+	{
+		FNeuralMorphEditorModel::CreateActors(InPersonaPreviewScene);
+
+		// Create the detail pose actor.
+		UWorld* const World = InPersonaPreviewScene->GetWorld();
+		if (!World)
+		{
+			return;
+		}
+
+		if (DetailPoseActor)
+		{
+			EditorActors.Remove(DetailPoseActor);
+			if (DetailPoseActor->GetActor())
+			{
+				World->DestroyActor(DetailPoseActor->GetActor(), true);
+			}
+		}
+		DetailPoseActor = CreateDetailPoseActor(World);
+		UpdateDetailPoseActor(*DetailPoseActor);
+
+		EditorActors.Add(DetailPoseActor);
+	}
+
+	void FDetailPoseEditorModel::InitEngineMorphTargets(const TArray<FVector3f>& Deltas)
+	{
+		TArray<FVector3f> FinalDeltas = Deltas;
+
+		// Add the deltas for all the detail pose morph targets.
+		// We can only do this when we actually trained the model.
+		// Also skip calculating them if we already have them calculated, unless there is some vertex count mismatch.
+		if (IsTrained() && (DetailPosesDeltas.IsEmpty() || DetailPosesDeltas.Num() % Model->GetNumBaseMeshVerts() != 0))
+		{
+			TArray<FDetailPoseModelDetailPose>& DetailPoses = GetDetailPoseModel()->GetDetailPoses();
+			CalculateDetailPoseDeltas(DetailPosesDeltas, DetailPoses);
+		}
+
+		// Strip the existing detail pose deltas if we have a trained network.
+		const UNeuralMorphNetwork* MorphNetwork = GetDetailPoseModel()->GetNeuralMorphNetwork();
+		if (MorphNetwork)
+		{
+			const int32 NumTrainedMorphTargets = MorphNetwork->GetNumOutputs() + 1;	// Add one to add the means morph target.
+			const TConstArrayView<FVector3f> MainMorphs(Deltas.GetData(), NumTrainedMorphTargets * Model->GetInputInfo()->GetNumBaseMeshVertices());
+			FinalDeltas = MainMorphs;
+		}
+		FinalDeltas += DetailPosesDeltas;
+		GetMorphModel()->SetMorphTargetDeltas(FinalDeltas);
+
+		FNeuralMorphEditorModel::InitEngineMorphTargets(FinalDeltas);
+	}
+
+	void FDetailPoseEditorModel::Tick(FEditorViewportClient* ViewportClient, float DeltaTime)
+	{
+		FMLDeformerEditorModel::Tick(ViewportClient, DeltaTime);
+		if (DetailPoseActor)
+		{
+			UpdateDetailPoseActor(*DetailPoseActor);
+			DetailPoseActor->Tick();
+		}
+	}
+
+	void FDetailPoseEditorModel::UpdateActorVisibility()
+	{
+		FNeuralMorphEditorModel::UpdateActorVisibility();
+
+		if (!DetailPoseActor)
+		{
+			return;
+		}
+
+		const UDetailPoseModelVizSettings* VizSettings = Cast<UDetailPoseModelVizSettings>(Model->GetVizSettings());
+		const bool bShowTestData = (VizSettings->GetVisualizationMode() == EMLDeformerVizMode::TestData);
+		const bool bInDefaultMode = GetEditor()->IsDefaultModeActive();
+		const bool bIsVisible = bShowTestData && VizSettings->GetDrawDetailPose() && DetailPoseActor->HasVisualMesh() && bInDefaultMode;
+
+		DetailPoseActor->SetVisibility(bIsVisible);
+	}
+}	// namespace UE::DetailPoseModel
+
+#undef LOCTEXT_NAMESPACE
diff --git a/Engine/Plugins/Experimental/Animation/MLDeformer/DetailPoseModel/Source/DetailPoseModelEditor/Private/DetailPoseModelDetails.cpp b/Engine/Plugins/Experimental/Animation/MLDeformer/DetailPoseModel/Source/DetailPoseModelEditor/Private/DetailPoseModelDetails.cpp
new file mode 100644
--- /dev/null
+++ b/Engine/Plugins/Experimental/Animation/MLDeformer/DetailPoseModel/Source/DetailPoseModelEditor/Private/DetailPoseModelDetails.cpp
@@ -0,0 +1,95 @@
+// Copyright Epic Games, Inc. All Rights Reserved.
+
+#include "DetailPoseModelDetails.h"
+#include "DetailPoseEditorModel.h"
+#include "NeuralMorphModelDetails.h"
+#include "MLDeformerGeomCacheHelpers.h"
+#include "DetailLayoutBuilder.h"
+#include "DetailCategoryBuilder.h"
+#include "DetailWidgetRow.h"
+#include "SWarningOrErrorBox.h"
+#include "Widgets/Layout/SBox.h"
+
+namespace UE::DetailPoseModel
+{
+	TSharedRef<IDetailCustomization> FDetailPoseModelDetails::MakeInstance()
+	{
+		return MakeShareable(new FDetailPoseModelDetails());
+	}
+
+	void FDetailPoseModelDetails::CustomizeDetails(IDetailLayoutBuilder& DetailBuilder)
+	{
+		using namespace UE::MLDeformer;
+
+		DetailLayoutBuilder = &DetailBuilder;
+
+		// Update the pointers and check if they are valid.
+		TArray<TWeakObjectPtr<UObject>> Objects;
+		DetailBuilder.GetObjectsBeingCustomized(Objects);
+		if (!UpdateMemberPointers(Objects))
+		{
+			return;
+		}
+
+		CreateCategories();
+
+		UDetailPoseModel* DetailPoseModel = GetCastModel();
+		check(DetailPoseModel);
+
+		// Check if the geom cache is compatible with the skeletal mesh.
+		const FText GeomCacheErrorText = GetGeomCacheErrorText(Model->GetSkeletalMesh(), DetailPoseModel->GetDetailPosesGeomCache());
+		FDetailWidgetRow& GeomCacheErrorRow = DetailPosesCategoryBuilder->AddCustomRow(FText::FromString("DetailPosesError"))
+			.Visibility(!GeomCacheErrorText.IsEmpty() ? EVisibility::Visible : EVisibility::Collapsed)
+			.WholeRowContent()
+			[
+				SNew(SBox)
+				.Padding(FMargin(0.0f, 4.0f))
+				[
+					SNew(SWarningOrErrorBox)
+					.MessageStyle(EMessageStyle::Error)
+					.Message(GeomCacheErrorText)
+				]
+			];
+
+		// Check the animation frame rate and frame count etc.
+		const FText GeomCacheAnimErrorText = GetGeomCacheAnimSequenceErrorText(DetailPoseModel->GetDetailPosesGeomCache(), DetailPoseModel->GetDetailPosesAnimSequence());
+		FDetailWidgetRow& GeomCacheAnimErrorRow = DetailPosesCategoryBuilder->AddCustomRow(FText::FromString("DetailPosesAnimError"))
+			.Visibility(!GeomCacheAnimErrorText.IsEmpty() ? EVisibility::Visible : EVisibility::Collapsed)
+			.WholeRowContent()
+			[
+				SNew(SBox)
+				.Padding(FMargin(0.0f, 4.0f))
+				[
+					SNew(SWarningOrErrorBox)
+					.MessageStyle(EMessageStyle::Error)
+					.Message(GeomCacheAnimErrorText)
+				]
+			];
+
+		// Check the mesh mappings between the geom cache and skeletal mesh.
+		const FText GeomCacheMappingErrorText = GetGeomCacheMeshMappingErrorText(DetailPoseModel->GetSkeletalMesh(), DetailPoseModel->GetDetailPosesGeomCache());
+		FDetailWidgetRow& GeomCacheMappingErrorRow = DetailPosesCategoryBuilder->AddCustomRow(FText::FromString("DetailPosesMappingError"))
+			.Visibility(!GeomCacheMappingErrorText.IsEmpty() ? EVisibility::Visible : EVisibility::Collapsed)
+			.WholeRowContent()
+			[
+				SNew(SBox)
+				.Padding(FMargin(0.0f, 4.0f))
+				[
+					SNew(SWarningOrErrorBox)
+					.MessageStyle(EMessageStyle::Error)
+					.Message(GeomCacheMappingErrorText)
+				]
+			];
+	}
+
+	void FDetailPoseModelDetails::CreateCategories()
+	{
+		FNeuralMorphModelDetails::CreateCategories();
+		DetailPosesCategoryBuilder = &DetailLayoutBuilder->EditCategory("Detail Poses", FText::GetEmpty(), ECategoryPriority::Default);
+	}
+
+	UDetailPoseModel* FDetailPoseModelDetails::GetCastModel() const
+	{
+		return Cast<UDetailPoseModel>(Model);
+	}
+}	// namespace UE::DetailPoseModel
diff --git a/Engine/Plugins/Experimental/Animation/MLDeformer/DetailPoseModel/Source/DetailPoseModelEditor/Private/DetailPoseModelEditorActor.cpp b/Engine/Plugins/Experimental/Animation/MLDeformer/DetailPoseModel/Source/DetailPoseModelEditor/Private/DetailPoseModelEditorActor.cpp
new file mode 100644
--- /dev/null
+++ b/Engine/Plugins/Experimental/Animation/MLDeformer/DetailPoseModel/Source/DetailPoseModelEditor/Private/DetailPoseModelEditorActor.cpp
@@ -0,0 +1,55 @@
+// Copyright Epic Games, Inc. All Rights Reserved.
+
+#include "DetailPoseModelEditorActor.h"
+#include "DetailPoseModelInstance.h"
+#include "MLDeformerComponent.h"
+#include "GeometryCacheComponent.h"
+#include "GeometryCache.h"
+
+namespace UE::DetailPoseModel
+{
+	FDetailPoseModelEditorActor::FDetailPoseModelEditorActor(const FConstructSettings& Settings)
+		: FMLDeformerGeomCacheActor(Settings)
+	{
+	}
+
+	FDetailPoseModelEditorActor::~FDetailPoseModelEditorActor() = default;
+
+	void FDetailPoseModelEditorActor::SetGeometryCache(UGeometryCache* InGeometryCache) const
+	{
+		if (GeomCacheComponent && GeomCacheComponent->GetGeometryCache() != InGeometryCache)
+		{
+			GeomCacheComponent->SetGeometryCache(InGeometryCache);
+		}
+	}
+
+	void FDetailPoseModelEditorActor::SetTrackedComponent(const UMLDeformerComponent* InComponent)
+	{
+		TrackedComponent = InComponent;
+	}
+
+	void FDetailPoseModelEditorActor::Tick() const
+	{
+		if (!GeomCacheComponent || !GeomCacheComponent->GetGeometryCache() || !TrackedComponent.IsValid())
+		{
+			return;
+		}
+
+		const UDetailPoseModelInstance* ModelInstance = Cast<UDetailPoseModelInstance>(TrackedComponent->GetModelInstance());
+		if (ModelInstance)
+		{
+			const UGeometryCache* GeomCache = GeomCacheComponent->GetGeometryCache();
+			const int32 FrameIndex = ModelInstance->GetBestDetailPoseIndex();
+			if (FrameIndex >= 0 && FrameIndex <= (GeomCache->GetEndFrame() - GeomCache->GetStartFrame() + 1))
+			{
+				GeomCacheComponent->SetManualTick(true);
+				const float NewTimeValue = GeomCacheComponent->GetTimeAtFrame(FrameIndex);
+				const float CurrentTime = GeomCacheComponent->GetAnimationTime();
+				if (FMath::Abs(NewTimeValue - CurrentTime) > 0.00001f)
+				{
+					GeomCacheComponent->TickAtThisTime(NewTimeValue, false, false, false);
+				}
+			}
+		}
+	}
+}	// namespace UE::DetailPoseModel
diff --git a/Engine/Plugins/Experimental/Animation/MLDeformer/DetailPoseModel/Source/DetailPoseModelEditor/Private/DetailPoseModelEditorModule.cpp b/Engine/Plugins/Experimental/Animation/MLDeformer/DetailPoseModel/Source/DetailPoseModelEditor/Private/DetailPoseModelEditorModule.cpp
new file mode 100644
--- /dev/null
+++ b/Engine/Plugins/Experimental/Animation/MLDeformer/DetailPoseModel/Source/DetailPoseModelEditor/Private/DetailPoseModelEditorModule.cpp
@@ -0,0 +1,64 @@
+// Copyright Epic Games, Inc. All Rights Reserved.
+
+#include "DetailPoseEditorModel.h"
+#include "DetailPoseModel.h"
+#include "DetailPoseModelDetails.h"
+#include "MLDeformerEditorModule.h"
+#include "Modules/ModuleManager.h"
+#include "EditorModeRegistry.h"
+#include "PropertyEditorModule.h"
+
+#define LOCTEXT_NAMESPACE "DetailPoseModelEditorModule"
+
+namespace UE::DetailPoseModel
+{
+	class FDetailPoseModelEditorModule
+		: public IModuleInterface
+	{
+	public:
+		// IModuleInterface overrides.
+		void StartupModule() override;
+		void ShutdownModule() override;
+		// ~END IModuleInterface overrides.
+	};
+}
+IMPLEMENT_MODULE(UE::DetailPoseModel::FDetailPoseModelEditorModule, DetailPoseModelEditor)
+
+namespace UE::DetailPoseModel
+{
+	using namespace UE::MLDeformer;
+
+	void FDetailPoseModelEditorModule::StartupModule()
+	{
+		// Register object detail customizations.
+		FPropertyEditorModule& PropertyModule = FModuleManager::LoadModuleChecked<FPropertyEditorModule>("PropertyEditor");
+		PropertyModule.RegisterCustomClassLayout("DetailPoseModel", FOnGetDetailCustomizationInstance::CreateStatic(&FDetailPoseModelDetails::MakeInstance));
+		PropertyModule.NotifyCustomizationModuleChanged();
+
+		// Register our custom ML deformer model to the model registry in the ML Deformer Framework.
+		FMLDeformerEditorModule& EditorModule = FModuleManager::LoadModuleChecked<FMLDeformerEditorModule>("MLDeformerFrameworkEditor");
+		FMLDeformerEditorModelRegistry& ModelRegistry = EditorModule.GetModelRegistry();
+		ModelRegistry.RegisterEditorModel(UDetailPoseModel::StaticClass(), FOnGetEditorModelInstance::CreateStatic(&FDetailPoseEditorModel::MakeInstance), /*ModelPriority*/10);
+	}
+
+	void FDetailPoseModelEditorModule::ShutdownModule()
+	{
+		// Unregister our ML Deformer model.
+		if (FModuleManager::Get().IsModuleLoaded(TEXT("MLDeformerFrameworkEditor")))
+		{
+			FMLDeformerEditorModule& EditorModule = FModuleManager::GetModuleChecked<FMLDeformerEditorModule>("MLDeformerFrameworkEditor");
+			FMLDeformerEditorModelRegistry& ModelRegistry = EditorModule.GetModelRegistry();
+			ModelRegistry.UnregisterEditorModel(UDetailPoseModel::StaticClass());
+		}
+
+		// Unregister object detail customizations for this model.
+		if (FModuleManager::Get().IsModuleLoaded(TEXT("PropertyEditor")))
+		{
+			FPropertyEditorModule& PropertyModule = FModuleManager::GetModuleChecked<FPropertyEditorModule>(TEXT("PropertyEditor"));
+			PropertyModule.UnregisterCustomClassLayout(TEXT("DetailPoseModel"));
+			PropertyModule.NotifyCustomizationModuleChanged();
+		}
+	}
+}	// namespace UE::DetailPoseModel
+
+#undef LOCTEXT_NAMESPACE
diff --git a/Engine/Plugins/Experimental/Animation/MLDeformer/DetailPoseModel/Source/DetailPoseModelEditor/Private/DetailPoseModelEditorStyle.cpp b/Engine/Plugins/Experimental/Animation/MLDeformer/DetailPoseModel/Source/DetailPoseModelEditor/Private/DetailPoseModelEditorStyle.cpp
new file mode 100644
--- /dev/null
+++ b/Engine/Plugins/Experimental/Animation/MLDeformer/DetailPoseModel/Source/DetailPoseModelEditor/Private/DetailPoseModelEditorStyle.cpp
@@ -0,0 +1,36 @@
+// Copyright Epic Games, Inc. All Rights Reserved.
+
+#include "DetailPoseModelEditorStyle.h"
+#include "Interfaces/IPluginManager.h"
+#include "Styling/SlateStyleRegistry.h"
+#include "Math/Color.h"
+#include "Misc/Paths.h"
+
+namespace UE::DetailPoseModel
+{
+	FDetailPoseModelEditorStyle::FDetailPoseModelEditorStyle()
+		: FSlateStyleSet("DetailPoseModelEditorStyle")
+	{
+		const FString ResourceDir = FPaths::Combine(IPluginManager::Get().FindPlugin(TEXT("DetailPoseModel"))->GetBaseDir(), TEXT("Resources"));
+		SetContentRoot(ResourceDir);
+
+		// Colors and sizes.
+		Set("DetailPoseModel.EditorActor.WireframeColor", FLinearColor(0.0f, 1.0f, 1.0f));
+		Set("DetailPoseModel.EditorActor.LabelColor", FLinearColor(0.0f, 1.0f, 1.0f));
+
+		FSlateStyleRegistry::RegisterSlateStyle(*this);
+	}
+
+	FDetailPoseModelEditorStyle::~FDetailPoseModelEditorStyle()
+	{
+		FSlateStyleRegistry::UnRegisterSlateStyle(*this);
+	}
+
+	FDetailPoseModelEditorStyle& FDetailPoseModelEditorStyle::Get()
+	{
+		static FDetailPoseModelEditorStyle Inst;
+		return Inst;
+	}
+}	// namespace UE::DetailPoseModel
+
+#undef MLDEFORMER_IMAGE_BRUSH
diff --git a/Engine/Plugins/Experimental/Animation/MLDeformer/DetailPoseModel/Source/DetailPoseModelEditor/Public/DetailPoseEditorModel.h b/Engine/Plugins/Experimental/Animation/MLDeformer/DetailPoseModel/Source/DetailPoseModelEditor/Public/DetailPoseEditorModel.h
new file mode 100644
--- /dev/null
+++ b/Engine/Plugins/Experimental/Animation/MLDeformer/DetailPoseModel/Source/DetailPoseModelEditor/Public/DetailPoseEditorModel.h
@@ -0,0 +1,75 @@
+// Copyright Epic Games, Inc. All Rights Reserved.
+
+#pragma once
+
+#include "CoreTypes.h"
+#include "NeuralMorphEditorModel.h"
+#include "DetailPoseModel.h"
+#include "Templates/UniquePtr.h"
+
+class UDetailPoseModelInstance;
+class UNeuralMorphNetworkInstance;
+
+namespace UE::DetailPoseModel
+{
+	using namespace UE::MLDeformer;
+
+	class FDetailPoseModelEditorActor;
+
+	/**
+	 * The editor model for the Detail Pose Model.
+	 */
+	class DETAILPOSEMODELEDITOR_API FDetailPoseEditorModel
+		: public UE::NeuralMorphModel::FNeuralMorphEditorModel
+	{
+	public:
+		// We need to implement this static MakeInstance method.
+		static FMLDeformerEditorModel* MakeInstance();
+
+		// FGCObject overrides.
+		FString GetReferencerName() const override final						{ return TEXT("FDetailPoseEditorModel"); }
+		// ~END FGCObject overrides.
+	
+		// FMLDeformerEditorModel overrides.
+		FMLDeformerTrainingInputAnim* GetTrainingInputAnim(int32 Index) const override final;
+		void UpdateTimelineTrainingAnimList() override final;
+		void OnPropertyChanged(FPropertyChangedEvent& PropertyChangedEvent) override final;
+		void OnPostTraining(ETrainingResult TrainingResult, bool bUsePartiallyTrainedWhenAborted) override final;
+		void InitInputInfo(UMLDeformerInputInfo* InputInfo) override final;
+		void CreateActors(const TSharedRef<IPersonaPreviewScene>& InPersonaPreviewScene) override final;
+		void Tick(FEditorViewportClient* ViewportClient, float DeltaTime) override final;
+		void UpdateActorVisibility() override final;
+		// ~END FMLDeformerEditorModel overrides.
+
+		// FMLDeformerMorphModelEditorModel overrides.
+		void InitEngineMorphTargets(const TArray<FVector3f>& Deltas) override final;
+		// ~END FMLDeformerMorphModelEditorModel overrides.
+
+		/** Get a pointer to the runtime model. */
+		UDetailPoseModel* GetDetailPoseModel() const							{ return Cast<UDetailPoseModel>(Model); }
+
+	private:
+		void CalculateDetailPoseDeltas(TArray<FVector3f>& OutDeltas, TArray<FDetailPoseModelDetailPose>& OutDetailPoses);
+		void CalculateDetailPoseDeltas(int32 DetailPoseIndex, FMLDeformerGeomCacheSampler& Sampler, UDetailPoseModelInstance* ModelInstance, TArrayView<FVector3f> OutDeltas, FDetailPoseModelDetailPose& OutDetailPose);
+
+		FDetailPoseModelEditorActor* CreateDetailPoseActor(UWorld* World) const;
+		void UpdateDetailPoseActor(FDetailPoseModelEditorActor& Actor) const;
+
+		UMLDeformerComponent* GetTestMLDeformerComponent() const;
+
+	private:
+		/** The training input anim object that holds the detail pose anim sequence and geom cache sequence. */
+		mutable TUniquePtr<FMLDeformerGeomCacheTrainingInputAnim> DetailPosesAnim;
+
+		/** The detail pose deltas, of size (NumDetailPoses * NumBaseMeshVerts). */
+		TArray<FVector3f> DetailPosesDeltas;
+
+		/** 
+		 * The editor actor that represents the detail pose actor. It is used to show the currently used detail pose. 
+		 * It's creation and destruction are managed internally.
+		 * @see CreateActors.
+		 */
+		FDetailPoseModelEditorActor* DetailPoseActor = nullptr;
+	};
+
+}	// namespace UE::DetailPoseModel
diff --git a/Engine/Plugins/Experimental/Animation/MLDeformer/DetailPoseModel/Source/DetailPoseModelEditor/Public/DetailPoseModelDetails.h b/Engine/Plugins/Experimental/Animation/MLDeformer/DetailPoseModel/Source/DetailPoseModelEditor/Public/DetailPoseModelDetails.h
new file mode 100644
--- /dev/null
+++ b/Engine/Plugins/Experimental/Animation/MLDeformer/DetailPoseModel/Source/DetailPoseModelEditor/Public/DetailPoseModelDetails.h
@@ -0,0 +1,46 @@
+// Copyright Epic Games, Inc. All Rights Reserved.
+
+#pragma once
+
+#include "CoreTypes.h"
+#include "NeuralMorphModelDetails.h"
+
+class IDetailChildrenBuilder;
+class IPropertyHandle;
+class UDetailPoseModel;
+
+namespace UE::DetailPoseModel
+{
+	class FDetailPoseEditorModel;
+
+	/**
+	 * The details customization for the model settings of the Detail Pose Model.
+	 * We implement a detail customization because we want to show some errors/warnings in case there are
+	 * any issues with the detail pose animation sequence or geometry cache.
+	 * Those errors could be like mismatching frame numbers etc.
+	 */
+	class DETAILPOSEMODELEDITOR_API FDetailPoseModelDetails
+		: public UE::NeuralMorphModel::FNeuralMorphModelDetails
+	{
+	public:
+		/** Makes a new instance of this detail layout class for a specific detail view requesting it. */
+		static TSharedRef<IDetailCustomization> MakeInstance();
+
+		// ILayoutDetails overrides.
+		virtual void CustomizeDetails(IDetailLayoutBuilder& DetailBuilder) override final;
+		virtual void CreateCategories() override final;
+		// ~END ILayoutDetails overrides.
+
+	private:
+		/**
+		 * Get a cast version of the UMLDeformerModel that we are showing details for.
+		 * This casts the model to the UDetailPoseModel class.
+		 * @result A pointer to the UDetailPoseModel that we are currently customizing, or nullptr when something is wrong.
+		 */
+		UDetailPoseModel* GetCastModel() const;
+
+	private:
+		/** The category that holds the Detail Pose settings in the model details panel. */
+		IDetailCategoryBuilder* DetailPosesCategoryBuilder = nullptr;
+	};
+}	// namespace UE::DetailPoseModel
diff --git a/Engine/Plugins/Experimental/Animation/MLDeformer/DetailPoseModel/Source/DetailPoseModelEditor/Public/DetailPoseModelEditorActor.h b/Engine/Plugins/Experimental/Animation/MLDeformer/DetailPoseModel/Source/DetailPoseModelEditor/Public/DetailPoseModelEditorActor.h
new file mode 100644
--- /dev/null
+++ b/Engine/Plugins/Experimental/Animation/MLDeformer/DetailPoseModel/Source/DetailPoseModelEditor/Public/DetailPoseModelEditorActor.h
@@ -0,0 +1,58 @@
+// Copyright Epic Games, Inc. All Rights Reserved.
+
+#pragma once
+
+#include "MLDeformerGeomCacheActor.h"
+
+class UGeometryCache;
+class UMLDeformerComponent;
+
+namespace UE::DetailPoseModel
+{
+	enum : int32
+	{
+		/** The detail pose actor, which shows the current closest matching detail pose. */
+		ActorID_DetailPoseActor = 6
+	};
+
+	/**
+	 * The editor actor for the Detail Pose model.
+	 * This represents the visual detail pose actor inside the ML Deformer asset editor's viewport.
+	 * It is used to show the current closest matching detail pose, compared to the current pose of our skeletal mesh component.
+	 */
+	class DETAILPOSEMODELEDITOR_API FDetailPoseModelEditorActor
+		: public UE::MLDeformer::FMLDeformerGeomCacheActor
+	{
+	public:
+		FDetailPoseModelEditorActor(const FConstructSettings& Settings);
+		virtual ~FDetailPoseModelEditorActor();
+
+		/**
+		 * Set the geometry cache asset to use to render in the editor viewport.
+		 * This should be the detail pose geometry cache.
+		 * @param InGeometryCache The geometry cache to use to render the detail poses.
+		 */
+		void SetGeometryCache(UGeometryCache* InGeometryCache) const;
+
+		/**
+		 * Specify which ML Deformer component to extract the current closest detail pose from.
+		 * The Tick method will use this and update the geometry cache components current time value based on this.
+		 * @param InComponent The ML Deformer Component to grab the current frame index from inside the detail pose geometry cache.
+		 * @see Tick.
+		 */
+		void SetTrackedComponent(const UMLDeformerComponent* InComponent);
+
+		/**
+		 * Tick the editor actor, which will basically tick the geometry cache component and set it to the
+		 * frame that is the current closest detail pose as reported by the tracked ML Deformer component.
+		 * The ML Deformer component (or well the UDetailPoseModelInstance inside it) stores the current closest detail pose.
+		 * We grab the frame number of that pose (the frame number inside the Geometry Cache), and then set the internal
+		 * geometry cache component's time to this frame number to display this pose.
+		 */
+		void Tick() const;
+	
+	private:
+		/** The ML Deformer component from which we grab the current frame value. */
+		TWeakObjectPtr<const UMLDeformerComponent> TrackedComponent;
+	};
+}	// namespace UE::DetailPoseModel
diff --git a/Engine/Plugins/Experimental/Animation/MLDeformer/DetailPoseModel/Source/DetailPoseModelEditor/Public/DetailPoseModelEditorStyle.h b/Engine/Plugins/Experimental/Animation/MLDeformer/DetailPoseModel/Source/DetailPoseModelEditor/Public/DetailPoseModelEditorStyle.h
new file mode 100644
--- /dev/null
+++ b/Engine/Plugins/Experimental/Animation/MLDeformer/DetailPoseModel/Source/DetailPoseModelEditor/Public/DetailPoseModelEditorStyle.h
@@ -0,0 +1,22 @@
+// Copyright Epic Games, Inc. All Rights Reserved.
+
+#pragma once
+
+#include "Styling/SlateStyle.h"
+
+namespace UE::DetailPoseModel
+{
+	/**
+	 * The styles used in the Detail Pose editor model.
+	 * This contains things like colors for the label of the Detail Pose model rendered in the viewport.
+	 */
+	class DETAILPOSEMODELEDITOR_API FDetailPoseModelEditorStyle
+		: public FSlateStyleSet
+	{
+	public:
+		FDetailPoseModelEditorStyle();
+		virtual ~FDetailPoseModelEditorStyle();
+
+		static FDetailPoseModelEditorStyle& Get();
+	};
+}	// namespace UE::DetailPoseModel
diff --git a/Engine/Plugins/Experimental/Animation/MLDeformer/DetailPoseModel/Source/DetailPoseModelEditor/Public/DetailPoseTrainingModel.h b/Engine/Plugins/Experimental/Animation/MLDeformer/DetailPoseModel/Source/DetailPoseModelEditor/Public/DetailPoseTrainingModel.h
new file mode 100644
--- /dev/null
+++ b/Engine/Plugins/Experimental/Animation/MLDeformer/DetailPoseModel/Source/DetailPoseModelEditor/Public/DetailPoseTrainingModel.h
@@ -0,0 +1,19 @@
+// Copyright Epic Games, Inc. All Rights Reserved.
+
+#pragma once
+
+#include "NeuralMorphTrainingModel.h"
+#include "DetailPoseTrainingModel.generated.h"
+
+/**
+ * The Training Model for the Detail Pose Model.
+ * Even though this is an empty class we still need it, as we inherit from it inside the 
+ * Python code. This allows us to separate it from the Neural Morph Model that we inherited from.
+ * We basically find an object of this class type and call the Train method on that in Python.
+ */
+UCLASS(Blueprintable)
+class DETAILPOSEMODELEDITOR_API UDetailPoseTrainingModel
+	: public UNeuralMorphTrainingModel
+{
+	GENERATED_BODY()
+};
diff --git a/Engine/Source/Editor/Persona/Private/AnimationEditorViewportClient.cpp b/Engine/Source/Editor/Persona/Private/AnimationEditorViewportClient.cpp
--- a/Engine/Source/Editor/Persona/Private/AnimationEditorViewportClient.cpp
+++ b/Engine/Source/Editor/Persona/Private/AnimationEditorViewportClient.cpp
@@ -2480,7 +2480,8 @@ void FAnimationViewportClient::FocusViewportOnPreviewMesh(bool bUseCustomCamera)
 
 float FAnimationViewportClient::GetFloorOffset() const
 {
-	USkeletalMesh* Mesh = GetPreviewScene()->GetPreviewMeshComponent()->GetSkeletalMeshAsset();
+	const USkeletalMeshComponent* SkelMeshComponent = GetPreviewScene()->GetPreviewMeshComponent();
+	USkeletalMesh* Mesh = SkelMeshComponent ? SkelMeshComponent->GetSkeletalMeshAsset() : nullptr;
 	if ( Mesh )
 	{
 		return Mesh->GetFloorOffset();
@@ -2491,7 +2492,8 @@ float FAnimationViewportClient::GetFloorOffset() const
 
 void FAnimationViewportClient::SetFloorOffset( float NewValue )
 {
-	USkeletalMesh* Mesh = GetPreviewScene()->GetPreviewMeshComponent()->GetSkeletalMeshAsset();
+	const USkeletalMeshComponent* SkelMeshComponent = GetPreviewScene()->GetPreviewMeshComponent();
+	USkeletalMesh* Mesh = SkelMeshComponent ? SkelMeshComponent->GetSkeletalMeshAsset() : nullptr;
 
 	if ( Mesh )
 	{
