[MLDeformer] Small PVS static analyzer check fix.

diff --git a/Engine/Plugins/Experimental/Animation/MLDeformer/DetailPoseModel/Source/DetailPoseModel/Private/DetailPoseModelInstance.cpp b/Engine/Plugins/Experimental/Animation/MLDeformer/DetailPoseModel/Source/DetailPoseModel/Private/DetailPoseModelInstance.cpp
--- a/Engine/Plugins/Experimental/Animation/MLDeformer/DetailPoseModel/Source/DetailPoseModel/Private/DetailPoseModelInstance.cpp
+++ b/Engine/Plugins/Experimental/Animation/MLDeformer/DetailPoseModel/Source/DetailPoseModel/Private/DetailPoseModelInstance.cpp
@@ -162,7 +162,7 @@ void UDetailPoseModelInstance::Execute(float ModelWeight)
 	UDetailPoseModelInputInfo* InputInfo = Cast<UDetailPoseModelInputInfo>(DetailPoseModel->GetInputInfo());
 	const int32 NumGlobalMorphs = InputInfo->GetNumGlobalMorphTargets();
 	UNeuralMorphNetwork* Network = DetailPoseModel->GetNeuralMorphNetwork();
-	const int32 NumNetworkOutputs = Network->GetNumMainOutputs();
+	const int32 NumNetworkOutputs = Network ? Network->GetNumMainOutputs() : 0;
 	if (!Network || Network->IsEmpty() || NumGlobalMorphs != NumNetworkOutputs || !NetworkInstance)
 	{
 		WeightData->ZeroWeights();
