[MLDeformer] Fixed a bug where the Neural Morph Model would actually launch the Detail Pose Model. This would break when the Neural Morph Model was set to Local mode. Also directly made the imports in python better, and fixed some comment in the python file.

#rnx

diff --git a/Engine/Plugins/Animation/MLDeformer/MLDeformerFramework/Source/MLDeformerFrameworkEditor/Public/MLDeformerEditorModel.h b/Engine/Plugins/Animation/MLDeformer/MLDeformerFramework/Source/MLDeformerFrameworkEditor/Public/MLDeformerEditorModel.h
--- a/Engine/Plugins/Animation/MLDeformer/MLDeformerFramework/Source/MLDeformerFrameworkEditor/Public/MLDeformerEditorModel.h
+++ b/Engine/Plugins/Animation/MLDeformer/MLDeformerFramework/Source/MLDeformerFrameworkEditor/Public/MLDeformerEditorModel.h
@@ -1305,7 +1305,7 @@ namespace UE::MLDeformer
 	{
 		// Find class, which will include our Python class, generated from the python script.
 		TArray<UClass*> TrainingModels;
-		GetDerivedClasses(TrainingModelClass::StaticClass(), TrainingModels);
+		GetDerivedClasses(TrainingModelClass::StaticClass(), TrainingModels, false);
 		if (TrainingModels.IsEmpty())
 		{
 			// We didn't define a derived class in Python.
diff --git a/Engine/Plugins/Experimental/Animation/MLDeformer/DetailPoseModel/Source/DetailPoseModelEditor/Private/DetailPoseEditorModel.cpp b/Engine/Plugins/Experimental/Animation/MLDeformer/DetailPoseModel/Source/DetailPoseModelEditor/Private/DetailPoseEditorModel.cpp
--- a/Engine/Plugins/Experimental/Animation/MLDeformer/DetailPoseModel/Source/DetailPoseModelEditor/Private/DetailPoseEditorModel.cpp
+++ b/Engine/Plugins/Experimental/Animation/MLDeformer/DetailPoseModel/Source/DetailPoseModelEditor/Private/DetailPoseEditorModel.cpp
@@ -7,6 +7,7 @@
 #include "DetailPoseModelEditorStyle.h"
 #include "DetailPoseModelEditorActor.h"
 #include "DetailPoseModelVizSettings.h"
+#include "DetailPoseTrainingModel.h"
 #include "MLDeformerComponent.h"
 #include "MLDeformerGeomCacheActor.h"
 #include "MLDeformerGeomCacheSampler.h"
@@ -443,6 +444,11 @@ namespace UE::DetailPoseModel
 
 		DetailPoseActor->SetVisibility(bIsVisible);
 	}
+
+	ETrainingResult FDetailPoseEditorModel::Train()
+	{
+		return TrainModel<UDetailPoseTrainingModel>(this);
+	}
 }	// namespace UE::DetailPoseModel
 
 #undef LOCTEXT_NAMESPACE
diff --git a/Engine/Plugins/Experimental/Animation/MLDeformer/DetailPoseModel/Source/DetailPoseModelEditor/Public/DetailPoseEditorModel.h b/Engine/Plugins/Experimental/Animation/MLDeformer/DetailPoseModel/Source/DetailPoseModelEditor/Public/DetailPoseEditorModel.h
--- a/Engine/Plugins/Experimental/Animation/MLDeformer/DetailPoseModel/Source/DetailPoseModelEditor/Public/DetailPoseEditorModel.h
+++ b/Engine/Plugins/Experimental/Animation/MLDeformer/DetailPoseModel/Source/DetailPoseModelEditor/Public/DetailPoseEditorModel.h
@@ -39,6 +39,7 @@ namespace UE::DetailPoseModel
 		void CreateActors(const TSharedRef<IPersonaPreviewScene>& InPersonaPreviewScene) override final;
 		void Tick(FEditorViewportClient* ViewportClient, float DeltaTime) override final;
 		void UpdateActorVisibility() override final;
+		ETrainingResult Train() override final;
 		// ~END FMLDeformerEditorModel overrides.
 
 		// FMLDeformerMorphModelEditorModel overrides.
